NetworkPrediction
* Sim Extrapolation support (extrapolate network sim by synthesizing input cmds)
* Split SyncTo into FinalizeFrame (called by net sim model) and PreSimSync (only used by flying movement sim right now).
* Setup for simulated proxy interpolation support

#rnx

diff --git a/Engine/Plugins/Runtime/NetworkPrediction/Source/NetworkPrediction/Private/Movement/FlyingMovementComponent.cpp b/Engine/Plugins/Runtime/NetworkPrediction/Source/NetworkPrediction/Private/Movement/FlyingMovementComponent.cpp
--- a/Engine/Plugins/Runtime/NetworkPrediction/Source/NetworkPrediction/Private/Movement/FlyingMovementComponent.cpp
+++ b/Engine/Plugins/Runtime/NetworkPrediction/Source/NetworkPrediction/Private/Movement/FlyingMovementComponent.cpp
@@ -166,7 +166,7 @@ void UFlyingMovementComponent::InitSyncState(FlyingMovement::FMoveState& OutSync
 	OutSyncState.Rotation = UpdatedComponent->GetComponentQuat().Rotator();	
 }
 
-void UFlyingMovementComponent::SyncTo(const FlyingMovement::FMoveState& SyncState)
+void UFlyingMovementComponent::PreSimSync(const FlyingMovement::FMoveState& SyncState)
 {
 	// Does checking equality make any sense here? This is unfortunate
 	if (UpdatedComponent->GetComponentLocation().Equals(SyncState.Location) == false || UpdatedComponent->GetComponentQuat().Rotator().Equals(SyncState.Rotation, FlyingMovement::ROTATOR_TOLERANCE) == false)
@@ -177,3 +177,8 @@ void UFlyingMovementComponent::SyncTo(const FlyingMovement::FMoveState& SyncStat
 		UpdatedComponent->ComponentVelocity = SyncState.Velocity;
 	}
 }
+
+void UFlyingMovementComponent::FinalizeFrame(const FlyingMovement::FMoveState& SyncState)
+{
+	PreSimSync(SyncState);
+}
\ No newline at end of file
diff --git a/Engine/Plugins/Runtime/NetworkPrediction/Source/NetworkPrediction/Private/Movement/FlyingMovementSimulation.cpp b/Engine/Plugins/Runtime/NetworkPrediction/Source/NetworkPrediction/Private/Movement/FlyingMovementSimulation.cpp
--- a/Engine/Plugins/Runtime/NetworkPrediction/Source/NetworkPrediction/Private/Movement/FlyingMovementSimulation.cpp
+++ b/Engine/Plugins/Runtime/NetworkPrediction/Source/NetworkPrediction/Private/Movement/FlyingMovementSimulation.cpp
@@ -182,7 +182,7 @@ namespace FlyingMovement
 
 		// After the rotation is known, we need to sync our driver to this state. This is unfortunate and probably not best for perf, but since the Driver owned
 		// primitive component is ultimately what does our scene queries, we have no choice: we must get the primitive component in the state we say it should be in.
-		Driver->SyncTo(OutputState);
+		Driver->PreSimSync(OutputState);
 
 		const FVector LocalSpaceMovementInput = OutputState.Rotation.RotateVector( InputCmd.MovementInput );
 	   	
diff --git a/Engine/Plugins/Runtime/NetworkPrediction/Source/NetworkPrediction/Private/Movement/ParametricMovement.cpp b/Engine/Plugins/Runtime/NetworkPrediction/Source/NetworkPrediction/Private/Movement/ParametricMovement.cpp
--- a/Engine/Plugins/Runtime/NetworkPrediction/Source/NetworkPrediction/Private/Movement/ParametricMovement.cpp
+++ b/Engine/Plugins/Runtime/NetworkPrediction/Source/NetworkPrediction/Private/Movement/ParametricMovement.cpp
@@ -19,6 +19,10 @@ static float DrawDebugDefaultLifeTime = 30.f;
 static FAutoConsoleVariableRef CVarDrawDebugDefaultLifeTime(TEXT("parametricmover.Debug.UseDrawDebug.DefaultLifeTime"),
 	DrawDebugDefaultLifeTime, TEXT("Use built in DrawDebug* functions for visual logging"), ECVF_Default);
 
+static int32 ForceNetUpdate = 0;
+static FAutoConsoleVariableRef CVarForceNetUpdate(TEXT("parametricmover.ForceNetUpdate"),
+	DrawDebugDefaultLifeTime, TEXT("Toggles calling ForceNetUpdate each tick on Parametric movers. This is only meant for testing/debugging."), ECVF_Default);
+
 }
 
 // -------------------------------------------------------------------------------------------------------------------------------
@@ -141,7 +145,7 @@ void UParametricMovementComponent::InitSyncState(ParametricMovement::FMoveState&
 	OutSyncState.PlayRate = 0.f;
 }
 
-void UParametricMovementComponent::SyncTo(const ParametricMovement::FMoveState& SyncState)
+void UParametricMovementComponent::FinalizeFrame(const ParametricMovement::FMoveState& SyncState)
 {
 	FTransform NewTransform;
 	GetTransformForTime(SyncState.Position, NewTransform);
@@ -207,5 +211,8 @@ void UParametricMovementComponent::TickComponent(float DeltaTime, enum ELevelTic
 	}
 
 	// TEMP
-	GetOwner()->ForceNetUpdate();
+	if (ParametricMoverCVars::ForceNetUpdate)
+	{
+		GetOwner()->ForceNetUpdate();
+	}
 }
\ No newline at end of file
diff --git a/Engine/Plugins/Runtime/NetworkPrediction/Source/NetworkPrediction/Public/Movement/FlyingMovement.h b/Engine/Plugins/Runtime/NetworkPrediction/Source/NetworkPrediction/Public/Movement/FlyingMovement.h
--- a/Engine/Plugins/Runtime/NetworkPrediction/Source/NetworkPrediction/Public/Movement/FlyingMovement.h
+++ b/Engine/Plugins/Runtime/NetworkPrediction/Source/NetworkPrediction/Public/Movement/FlyingMovement.h
@@ -119,8 +119,13 @@ namespace FlyingMovement
 		class IMovementDriver : public IDriver
 		{
 		public:
-			// FlyingMovement only really needs the base driver functions for moving a primitive component around.
+			// Interface for moving the collision component around
 			virtual IBaseMovementDriver& GetBaseMovementDriver() = 0;
+
+			// Called prior to running the sim to make sure to make sure the collision component is in the right place. 
+			// This is unfortunate and not good, but is needed to ensure our collision and world position have not been moved out from under us.
+			// Refactoring primitive component movement to allow the sim to do all collision queries outside of the component code would be ideal.
+			virtual void PreSimSync(const FMoveState& SyncState) = 0;
 		};
 
 		/** Main update function */
@@ -160,7 +165,8 @@ public:
 	IBaseMovementDriver& GetBaseMovementDriver() override final { return *static_cast<IBaseMovementDriver*>(this); }
 
 	void InitSyncState(FlyingMovement::FMoveState& OutSyncState) const override;
-	void SyncTo(const FlyingMovement::FMoveState& SyncState) override;
+	void PreSimSync(const FlyingMovement::FMoveState& SyncState) override;
+	void FinalizeFrame(const FlyingMovement::FMoveState& SyncState) override;
 
 protected:
 
diff --git a/Engine/Plugins/Runtime/NetworkPrediction/Source/NetworkPrediction/Public/Movement/ParametricMovement.h b/Engine/Plugins/Runtime/NetworkPrediction/Source/NetworkPrediction/Public/Movement/ParametricMovement.h
--- a/Engine/Plugins/Runtime/NetworkPrediction/Source/NetworkPrediction/Public/Movement/ParametricMovement.h
+++ b/Engine/Plugins/Runtime/NetworkPrediction/Source/NetworkPrediction/Public/Movement/ParametricMovement.h
@@ -146,7 +146,7 @@ class NETWORKPREDICTION_API UParametricMovementComponent : public UBaseMovementC
 
 	// Base TNetworkModelSimulation driver
 	void InitSyncState(ParametricMovement::FMoveState& OutSyncState) const override;
-	void SyncTo(const ParametricMovement::FMoveState& SyncState) override;
+	void FinalizeFrame(const ParametricMovement::FMoveState& SyncState) override;
 
 	// Base Movement Driver
 	IBaseMovementDriver& GetBaseMovementDriver() override final { return *static_cast<IBaseMovementDriver*>(this); }
diff --git a/Engine/Plugins/Runtime/NetworkPrediction/Source/NetworkPrediction/Public/NetworkSimulationModel.h b/Engine/Plugins/Runtime/NetworkPrediction/Source/NetworkPrediction/Public/NetworkSimulationModel.h
--- a/Engine/Plugins/Runtime/NetworkPrediction/Source/NetworkPrediction/Public/NetworkSimulationModel.h
+++ b/Engine/Plugins/Runtime/NetworkPrediction/Source/NetworkPrediction/Public/NetworkSimulationModel.h
@@ -123,6 +123,19 @@ struct FVisualLoggingParameters
 	static NETWORKPREDICTION_API FColor DebugColors[(int32)EVisualLoggingContext::MAX];
 };
 
+UENUM()
+enum class ESimulatedUpdateMode : uint8
+{
+	Interpolate,		// Update from previous to current known states from the server. This puts the simulation further "behind" due to having to buffer the known state (but is never "wrong" and doesn't require a simulation update)
+	Extrapolate,		// Extrapolate the simulation once per local frame, by synthesizing (guessing) input commands
+	ForwardPredict		// Predict the simulation ahead. For a simulated proxy to do this, it must be tied to an autonomous proxy
+};
+
+inline FString LexToString(ESimulatedUpdateMode A)
+{
+	return *UEnum::GetValueAsString(TEXT("NetworkPrediction.ESimulatedUpdateMode"), A);
+}
+
 // -------------------------------------------------------------------------------------------------------------------------------
 // Interface that the proxy talks to. This is what will implement the replication.
 // -------------------------------------------------------------------------------------------------------------------------------
diff --git a/Engine/Plugins/Runtime/NetworkPrediction/Source/NetworkPrediction/Public/NetworkSimulationModelTemplates.h b/Engine/Plugins/Runtime/NetworkPrediction/Source/NetworkPrediction/Public/NetworkSimulationModelTemplates.h
--- a/Engine/Plugins/Runtime/NetworkPrediction/Source/NetworkPrediction/Public/NetworkSimulationModelTemplates.h
+++ b/Engine/Plugins/Runtime/NetworkPrediction/Source/NetworkPrediction/Public/NetworkSimulationModelTemplates.h
@@ -403,7 +403,7 @@ struct TReplicator_BasicReconciliar
 				// We don't have corresponding local state. There are two cases:
 				if (NetworkSimulationModelCVars::EnableLocalPrediction) // Fixme: this is awkward. Expected if player prediction is disabled but coupling like this feels bad.
 				{
-					UE_LOG(LogNetworkSim, Warning, TEXT("::NetSerialize Fault: MotionStateBuffer does not contain data for frame %d. [%d-%d]"), SerializedHeadKeyframe, SyncBuffer.GetTailKeyframe(), SyncBuffer.GetHeadKeyframe());
+					UE_LOG(LogNetworkSim, Warning, TEXT("::NetSerialize Fault: SyncBuffer does not contain data for frame %d. [%d-%d]"), SerializedHeadKeyframe, SyncBuffer.GetTailKeyframe(), SyncBuffer.GetHeadKeyframe());
 				}
 
 				// Case 1: the serialized state is older than what we've kept in our buffer. A bigger buffer would solve this! (at the price of more resimulated frames to recover when this happens)
@@ -476,7 +476,7 @@ struct TReplicator_BasicReconciliar
 		// Copy authoritative state over client state (FIXME: we may want to store this off historically somewhere? Or will VLog be enough for debugging?)
 		*ClientSyncState = *ServerState;
 		
-		// Resimulate up to our latest MotionStateBuffer frame. Note that we may have unprocessed user commands in the command buffer at this point.
+		// Resimulate up to our latest SyncBuffer frame. Note that we may have unprocessed user commands in the command buffer at this point.
 		const int32 LatestKeyframe = SyncBuffer.GetHeadKeyframe();		
 		for (int32 Keyframe = ReconciliationKeyframe+1; Keyframe <= LatestKeyframe; ++Keyframe)
 		{
@@ -554,8 +554,8 @@ public:
 	class IDriver
 	{
 	public:
-		virtual void InitSyncState(TSyncedState& OutSyncState) const = 0;
-		virtual void SyncTo(const TSyncedState& SyncState) = 0;
+		virtual void InitSyncState(TSyncedState& OutSyncState) const = 0;	// Called to create initial value of the sync state.
+		virtual void FinalizeFrame(const TSyncedState& SyncState) = 0;		// Called only from the Network Sim at the end of the sim frame when there is new sync data.
 	};
 
 	struct FTickParameters
@@ -586,10 +586,12 @@ public:
 		// How many commands we are allowed to process right now. This will need to be built out a bit more to handle shifting network conditions etc. For now, just a cvar.
 		int32 AllowCmds = NetworkSimulationModelCVars::MaxInputCmdsFrame;
 
+		switch (Parameters.Role) {
+
 		// -------------------------------------------------------------------------------------------------------------------------------------------------
 		//														Reconciliation
 		// -------------------------------------------------------------------------------------------------------------------------------------------------
-		if (Parameters.Role == ROLE_AutonomousProxy)
+		case ROLE_AutonomousProxy:
 		{
 			// Client: don't allow buffering and ensure continuous stream of inputs. This will need to be built out a bit more. The main points are:
 			// -Don't buffer input locally (unless we explicitly opt into it as a lag hiding measurement). Don't want to accidentally introduce application latency.
@@ -633,7 +635,7 @@ public:
 				checkf(InputBuffer.IsValidKeyframe(SyncBuffer.GetHeadKeyframe()), TEXT("MotionState and InputCmd buffers are out of system. LastProcessedInputKeyframe: %d {%s} vs {%s}"),
 					LastProcessedInputKeyframe, *SyncBuffer.GetBasicDebugStr(), *InputBuffer.GetBasicDebugStr());
 
-				UE_LOG(LogNetworkSim, Warning, TEXT("Skipping local input frames because we have newer data in MotionStateBuffer. LastProcessedInputKeyframe: %d. {%s} {%s}"),
+				UE_LOG(LogNetworkSim, Warning, TEXT("Skipping local input frames because we have newer data in SyncBuffer. LastProcessedInputKeyframe: %d. {%s} {%s}"),
 					LastProcessedInputKeyframe, *SyncBuffer.GetBasicDebugStr(), *InputBuffer.GetBasicDebugStr());
 
 				LastProcessedInputKeyframe = SyncBuffer.GetHeadKeyframe();
@@ -644,9 +646,11 @@ public:
 				AllowCmds = 0; // Don't process any commands this frame
 				LastProcessedInputKeyframe = InputBuffer.GetHeadKeyframe(); // Increment so we can accept a new command next frame
 			}
+
+			break;
 		}
 
-		if (Parameters.Role == ROLE_Authority)
+		case ROLE_Authority:
 		{
 			if ( LastProcessedInputKeyframe+1 < InputBuffer.GetTailKeyframe() )
 			{
@@ -654,8 +658,33 @@ public:
 				UE_LOG(LogNetworkSim, Warning, TEXT("::Tick missing inputcmds. LastProcessedInputKeyframe: %d. %s"), LastProcessedInputKeyframe, *InputBuffer.GetBasicDebugStr());
 				LastProcessedInputKeyframe = InputBuffer.GetTailKeyframe()+1;
 			}
+			break;
+		}
+
+		case ROLE_SimulatedProxy:
+		{
+			// (TODO: only if extrapolating, which is currently not a setting anywhere)
+			// Take last user command and reuse it, with this frame's frame delta 
+			if (TInputCmd* SynthesizedCmd = GetNextInputForWrite(Parameters.LocalDeltaTimeSeconds))
+			{
+				// Having some way to base the Synthesized off the last command, and be customizable per type, would be powerful
+			}
+
+			if (LastProcessedInputKeyframe+1 == SyncBuffer.GetHeadKeyframe())
+			{
+				// UE_LOG(LogNetworkSim, Warning, TEXT("Skipping synthetic input frame because we have data from server (%d / %d)"), LastProcessedInputKeyframe, SyncBuffer.GetHeadKeyframe());
+
+				// We've received new sync buffer data from replication and haven't processed a user command for it
+				// (we don't need to process a user command for it, because this is the authoritative state from the server)
+				LastProcessedInputKeyframe++;
+				AllowCmds = 0;
+			}
+
+			break;
 		}
 
+		} // end switch
+
 		// -------------------------------------------------------------------------------------------------------------------------------------------------
 		//														Input Processing
 		// -------------------------------------------------------------------------------------------------------------------------------------------------
@@ -724,7 +753,7 @@ public:
 			//check(SyncBuffer.GetNumValidElements() > 0);
 			if (SyncBuffer.GetNumValidElements() > 0)
 			{
-				Driver->SyncTo(*SyncBuffer.GetElementFromHead(0));
+				Driver->FinalizeFrame(*SyncBuffer.GetElementFromHead(0));
 			}
 		}
 
@@ -830,6 +859,7 @@ public:
 		{
 			// Only return a cmd if we have processed the last one. This is a bit heavy handed but is a good practice to start with. We want buffering of input to be very explicit, never something that accidently happens.
 			TInputCmd* Next = InputBuffer.GetWriteNext();
+			*Next = TInputCmd();
 			Next->FrameDeltaTime = DeltaTime;
 			return Next;
 		}
@@ -843,7 +873,6 @@ public:
 
 		// We want to start with an empty command in the input buffer. See notes in input buffer processing function.
 		*InputBuffer.GetWriteNext() = TInputCmd();
-
 		LastLocalInputGFrameNumber = 0;
 	}
 
diff --git a/Engine/Plugins/Runtime/NetworkPrediction/readme.txt b/Engine/Plugins/Runtime/NetworkPrediction/readme.txt
--- a/Engine/Plugins/Runtime/NetworkPrediction/readme.txt
+++ b/Engine/Plugins/Runtime/NetworkPrediction/readme.txt
@@ -1,8 +1,7 @@
-
 // -------------------------------------------------------------------------------------------------------------------------------------------------
 //	Network Prediction Plugin
 // -------------------------------------------------------------------------------------------------------------------------------------------------
-
+ 
 High level overview:
 -This is a *WIP* plugin that aims to generalize and improve our player prediction systems.
 -Long term we hope this replaces UCharacterMovementComponent and its networking system (ServerMove, ClientAck*, ClientAdjust*, etc).
@@ -10,85 +9,105 @@ High level overview:
 -NetworkPrediction plugin is the core generalized system that you need to enable to use anything.
 	-For now at least, new movement classes built on top of this will live here but could one day be moved out into a dedicated movement system plugin.
 -NetworkPredictionExtras plugin contains some examples that may be helpful in learning the system. It is not inteded to be used directly in shipping games.
-
+ 
 High level goals:
 -Make working with predicted gameplay systems such as character movement easier.
 -Give predictived gameplay systems the ability to interop with each other (E.g, CharacterMovement and Gameplay Abilities).
 -Provide better tooling to understand what is going on in these complex predicted gameplay systems.
 -This all leads to: ability to develop richer and more interactive gameplay systems that work in multiplayer games.
-
+ 
 What this is NOT:
 -We are not moving to a global fixed ticking system. Though these systems have huge advantages, we don't think it is a good universal solution for the engine.
 -That said, we recognize that it is the right solution for many things and we want to enable support for that. Nothing in this system /requires/ dynamic ticking, but it is something we feel the need to support.
 -To be clear: widely different, variable, fluctuating tick rates between clients and server is basically the root problems we are fighting here.
--Related, we are also not all in on a "predict everything: seemlessly and automtically" solution. We still feel player prediction is best kept to a minimum (meaning: predict the minimum amount of stuff you can get away with).
--But, we are tring to make it way easier to write predictive gameplay systems. Easier to write, to debug and reason about.
+-Related, we are also not all in on a "predict everything: seamlessly and automatically" solution. We still feel player prediction is best kept to a minimum (meaning: predict the minimum amount of stuff you can get away with).
+-But, we are trying to make it way easier to write predictive gameplay systems. Easier to write, to debug and reason about.
 -We do hope the new movement system will be automatic in this sense: if you write a new movement mode within the movement simulation class, things should "just work". (But we are not trying to make your entire game predicted)
-
-
+ 
+ 
 Getting Started:
 -Both NetworkPrediction and NetworkPredictionExtras are disabled by default. You will need to enable them for your project (manually edit .uproject or do so through editor plugin screen).
 -Once NetworkPredictionExtras is loaded, /NetworkPredictionExtras/Content/TestMap.umap can be loaded (must enable 'Show Engine Content' AND 'Show Plugin Content' in content browser!).
 -The "MockNetworkSimulation" can be tested anywhere just by typing "mns.Spawn".
-
+ 
 Code to look at:
 -MockNetworkSimulation.h: Good place to start to see a bare minimum example of a predicted network simulation. Both the simulation and actor component live in this header.
 -NetworkPredictionExtrasFlyingPawn.h: Starting point to see how a pawn using flying movement is setup. In NetworkPredictionExtras.
 -FlyingMovement.h: The actual implementation of the flying movement system. In NetworkPrediction.
 -NetworkSimulationModelTemplates.h: this is the lowest level guts of the system. This is what every simulation ultimately uses to stay in sync.
-
+ 
 Extremely high level technical overview:
--Concept of "NetworkedSimulationModel": defining a tightly constrained (gameplay) simulation that can be run predictively on clients and authoratatively on server. E.g., a "Movement System" is a NetworkedSimulationModel.
+-Concept of "NetworkedSimulationModel": defining a tightly constrained (gameplay) simulation that can be run predictively on clients and authoritatively on server. E.g., a "Movement System" is a NetworkedSimulationModel.
 -The simulation has tightly defined input and output. At the lowest level we have generic buffers of structs that are synchronized:
 	Input Buffer: the data that the controlling client generates
 	Sync Buffer: the data that evolves frame to frame and is produced by ticking the simulation
-	Aux Buffer: additional data that can change, can be predicted, etc but does not necessarily get updated frame to frame (e.g, somethign else usually updates it).
+	Aux Buffer: additional data that can change, can be predicted, etc but does not necessarily get updated frame to frame (e.g, something else usually updates it).
 -Implicit correction: all prediction and correction logic is now client side. Server just tells clients what the authority state was and it is up to them to reconcile.
--Single code path: the core network simulation has one "Update" function that does not branch based on authority/prediction/simulated etc. "Write once" is the goal for gamplay code here.
-
+-Single code path: the core network simulation has one "Update" function that does not branch based on authority/prediction/simulated etc. "Write once" is the goal for gameplay code here.
+ 
 UE4 Networking info (How this plugs into UE4 networking)
 -See UNetworkPredictionComponent for the glue.
 -"Replication Proxy" (FReplicationProxy): struct that points to the network simulation and has policies about how it replicates and to who. Custom NetSerialize.
 -E.g, there is a replication proxy for owner, non owners, replays, and debug. All point to the same underlying data.
 -Still using actor replication to determine when to replicate.
-
-
-
+ 
+ 
+ 
 TODO: Major missing elements
 -Aux Buffer is not implemented, but referenced in a few places.
 -"Events" into and out of the system are not in place. That is, something happens in your movement sim that needs to call outside code/generate an event. A system is needed to track this stuff rather than just calling directly into the handler.
 -Interpolation/smoothing layer (visual, non simulation affecting smoothing to help with corrections and simulated proxies)
--No optimizations (bandwith, cpu) have been done.
+-No optimizations (bandwidth, cpu) have been done.
 -[Movement specific]: must lock down SetActorLocation/Rotation API so that movement state can be changed out from underneath the system.
-
+ 
 Short term road map:
 -The focus right now is on the generalized prediction system, rather than actual movement code.
 -We are exploring "forward predicting non autonomous proxies" with the parametric mover. Unclear how deep we will go with this.
 -Then Aux buffer and "events"
 -Then mixing in ability system + other scripting options
-
+ 
 -Once everything feels solid we will go deeper on movement and build out a more extendable movement system.
 -E.g, We don't intend all users of the engine to write their own network simulation model movement system.
--We hope to provide a generalized movement system that is easy to extend in itself without knowing all the details of the NetwokrPrediction plugin.
+-We hope to provide a generalized movement system that is easy to extend in itself without knowing all the details of the NetworkPrediction plugin.
 -We intend to support RootMotion (Animation) / RootMotionSources (non anim) in the new movement system.
-
-
-
+ 
+ 
+// ----------------------------------------------------------------------------------------------------------
+// Glossary
+// ----------------------------------------------------------------------------------------------------------
+ 
+Reconcile: When a predicting client receives authoritative data from the server that conflicts with what they previously predicted: sorting this out is called "reconciling". Usually means you will resimulate.
+Synthesize: In general means "making shit up". For example, "synthesizing user commands" would mean "creating user commands that were not generated by a player, but instead guessing what it might be".
+ 
+ 
+Client Prediction: using new, locally created information (input cmd) to evolve the game state ahead of the authoritative server.
+(Client) Extrapolation: Taking the latest received state from the authoritative server and "guessing" how it continues to evolve. The key difference is you are not using new information that the server does not have yet.
+ 
+Forward Predicting: Using Client Prediction to predict ahead N frames, that is proportional to the ping time with the server.
+(Sim) Extrapolation: Using the network sim model to extrapolate frames, by (usually) synthesizing input cmds. Extrapolation is "one frame at a time" rather than "N frames based on ping".
+Interpolation: Interpolating between known states. E.g, you are never guessing at future state.
+	Note: We may technically allow interpolation past "100%", which would technically make it extrapolation (you are guessing at state you don't have yet). 
+	This would still be distinct from the above "Sim" Extrapolation where the network sim is used to generate future data instead of simple "interpolation" algorithms.
+Smoothing: Taking the output of the simulation and applying an additional layer of "ease in/ease out" updates. Smoothing would use things like spring equations and "lag constants" rather than being tied to the actual simulation itself.
+ 
+ 
 // ----------------------------------------------------------------------------------------------------------
 // Release notes
 // ----------------------------------------------------------------------------------------------------------
-
+ 
 Update (8-2-19)
 -We want to take a look non player controlled simulations, such as doors, elevators, "pushers" before going deeper on ability system or movement.
 -Basic ParametricMoverment system checked in. This is also in an incomplete state. ("Pushing" the flying movement component does not work and won't be the focus for now until we are looking for closely at movement specifically).
 -The short term goal here is to be able to forward predict these even if they are simulated proxies. When/why/how is very tbd, but we want to see what it looks like.
-
-
+ 
+ 
 Current State (7-25-19)
 -Initial public check in. Core simulation model structure is in place. Lots of pieces missing but overall system is roughed in.
 -Two main examples right now: Flying Movement and 'MockNetworkSimulation' (in NetworkPredictionExtras)
 -Flying Movement: a port of FloatingPawnMovement into the new system. Essentially just a basic flying movement system.
 -MockNetworkSimulation: a very basic minimal "simulation" that demonstrates how the system works. Not tied to movement or anything physical: just a running counter/sum.
-
+ 
+ 
+ 
 
 
diff --git a/Engine/Plugins/Runtime/NetworkPredictionExtras/Source/NetworkPredictionExtras/Private/MockNetworkSimulation.cpp b/Engine/Plugins/Runtime/NetworkPredictionExtras/Source/NetworkPredictionExtras/Private/MockNetworkSimulation.cpp
--- a/Engine/Plugins/Runtime/NetworkPredictionExtras/Source/NetworkPredictionExtras/Private/MockNetworkSimulation.cpp
+++ b/Engine/Plugins/Runtime/NetworkPredictionExtras/Source/NetworkPredictionExtras/Private/MockNetworkSimulation.cpp
@@ -207,7 +207,7 @@ void UMockNetworkSimulationComponent::InitSyncState(FMockSyncState& OutSyncState
 	OutSyncState.Total = MockValue;
 }
 
-void UMockNetworkSimulationComponent::SyncTo(const FMockSyncState& SyncState)
+void UMockNetworkSimulationComponent::FinalizeFrame(const FMockSyncState& SyncState)
 {
 	MockValue = SyncState.Total;
 }
diff --git a/Engine/Plugins/Runtime/NetworkPredictionExtras/Source/NetworkPredictionExtras/Public/MockNetworkSimulation.h b/Engine/Plugins/Runtime/NetworkPredictionExtras/Source/NetworkPredictionExtras/Public/MockNetworkSimulation.h
--- a/Engine/Plugins/Runtime/NetworkPredictionExtras/Source/NetworkPredictionExtras/Public/MockNetworkSimulation.h
+++ b/Engine/Plugins/Runtime/NetworkPredictionExtras/Source/NetworkPredictionExtras/Public/MockNetworkSimulation.h
@@ -163,7 +163,7 @@ private:
 	
 public:
 	void InitSyncState(FMockSyncState& OutSyncState) const override;
-	void SyncTo(const FMockSyncState& SyncState) override;
+	void FinalizeFrame(const FMockSyncState& SyncState) override;
 	virtual UWorld* GetDriverWorld() const override final { return GetWorld(); }
 	virtual UObject* GetVLogOwner() const override final;
 	virtual FTransform GetDebugWorldTransform() const override final;
