Updated how the OodleTrainerCommandlet is specified on the commandline, due to a recent unfixed Engine bug with commandlets (there was a JIRA assigned to Core for that).

diff --git a/Engine/Plugins/Runtime/PacketHandlers/CompressionComponents/Oodle/BatchDebugDump.bat b/Engine/Plugins/Runtime/PacketHandlers/CompressionComponents/Oodle/BatchDebugDump.bat
--- a/Engine/Plugins/Runtime/PacketHandlers/CompressionComponents/Oodle/BatchDebugDump.bat
+++ b/Engine/Plugins/Runtime/PacketHandlers/CompressionComponents/Oodle/BatchDebugDump.bat
@@ -49,7 +49,7 @@ goto End
 
 
 :AutoGenDictionaries
-set DebugDumpParms=-run=OodleHandlerComponent.OodleTrainerCommandlet DebugDump
+set DebugDumpParms=-run=OodleTrainerCommandlet DebugDump
 set PreDumpCmdLine=%GameName:"=% %DebugDumpParms% %OutputDir:"=% %CaptureDir:"=%
 set PostDumpCmdLine=-forcelogflush
 
diff --git a/Engine/Plugins/Runtime/PacketHandlers/CompressionComponents/Oodle/BuildDictionary.bat b/Engine/Plugins/Runtime/PacketHandlers/CompressionComponents/Oodle/BuildDictionary.bat
--- a/Engine/Plugins/Runtime/PacketHandlers/CompressionComponents/Oodle/BuildDictionary.bat
+++ b/Engine/Plugins/Runtime/PacketHandlers/CompressionComponents/Oodle/BuildDictionary.bat
@@ -29,7 +29,7 @@ echo.
 
 
 :AutoGenDictionaries
-set AutoGenDictionariesParms=-run=OodleHandlerComponent.OodleTrainerCommandlet AutoGenerateDictionaries
+set AutoGenDictionariesParms=-run=OodleTrainerCommandlet AutoGenerateDictionaries
 set FinalGenCmdLine=%GameName:"=% %AutoGenDictionariesParms% -forcelogflush
 
 echo Executing dictionary generation commandlet - commandline:
diff --git a/Engine/Plugins/Runtime/PacketHandlers/CompressionComponents/Oodle/BuildDictionary_Generic.bat b/Engine/Plugins/Runtime/PacketHandlers/CompressionComponents/Oodle/BuildDictionary_Generic.bat
--- a/Engine/Plugins/Runtime/PacketHandlers/CompressionComponents/Oodle/BuildDictionary_Generic.bat
+++ b/Engine/Plugins/Runtime/PacketHandlers/CompressionComponents/Oodle/BuildDictionary_Generic.bat
@@ -53,7 +53,7 @@ set /p DirectoryRoot=Type the root directory where the capture files are located
 echo.
 
 :AutoGenDictionaries
-set AutoGenDictionariesParms=-run=OodleHandlerComponent.OodleTrainerCommandlet GenerateDictionary %DictionaryOutput% %FileFilter% %ChangelistFilter% all %DirectoryRoot%
+set AutoGenDictionariesParms=-run=OodleTrainerCommandlet GenerateDictionary %DictionaryOutput% %FileFilter% %ChangelistFilter% all %DirectoryRoot%
 set FinalGenCmdLine=%GameName:"=% %AutoGenDictionariesParms% -forcelogflush
 
 echo Executing dictionary generation commandlet - commandline:
diff --git a/Engine/Plugins/Runtime/PacketHandlers/CompressionComponents/Oodle/BuildDictionary_TestCompression.bat b/Engine/Plugins/Runtime/PacketHandlers/CompressionComponents/Oodle/BuildDictionary_TestCompression.bat
--- a/Engine/Plugins/Runtime/PacketHandlers/CompressionComponents/Oodle/BuildDictionary_TestCompression.bat
+++ b/Engine/Plugins/Runtime/PacketHandlers/CompressionComponents/Oodle/BuildDictionary_TestCompression.bat
@@ -29,7 +29,7 @@ echo.
 
 
 :AutoGenDictionaries
-set AutoGenDictionariesParms=-run=OodleHandlerComponent.OodleTrainerCommandlet AutoGenerateDictionaries
+set AutoGenDictionariesParms=-run=OodleTrainerCommandlet AutoGenerateDictionaries
 set FinalGenCmdLine=%GameName:"=% %AutoGenDictionariesParms% -forcelogflush -CompressionTest
 
 echo Executing dictionary generation commandlet - commandline:
diff --git a/Engine/Plugins/Runtime/PacketHandlers/CompressionComponents/Oodle/DumpOodleSDKPackets.bat b/Engine/Plugins/Runtime/PacketHandlers/CompressionComponents/Oodle/DumpOodleSDKPackets.bat
--- a/Engine/Plugins/Runtime/PacketHandlers/CompressionComponents/Oodle/DumpOodleSDKPackets.bat
+++ b/Engine/Plugins/Runtime/PacketHandlers/CompressionComponents/Oodle/DumpOodleSDKPackets.bat
@@ -29,7 +29,7 @@ echo.
 
 
 :AutoGenDictionaries
-set AutoGenDictionariesParms=-run=OodleHandlerComponent.OodleTrainerCommandlet AutoGenerateDictionaries
+set AutoGenDictionariesParms=-run=OodleTrainerCommandlet AutoGenerateDictionaries
 set FinalGenCmdLine=%GameName:"=% %AutoGenDictionariesParms% -forcelogflush -OodleDebugDump
 
 echo Executing dictionary generation commandlet - commandline:
diff --git a/Engine/Plugins/Runtime/PacketHandlers/CompressionComponents/Oodle/FirstRun.bat b/Engine/Plugins/Runtime/PacketHandlers/CompressionComponents/Oodle/FirstRun.bat
--- a/Engine/Plugins/Runtime/PacketHandlers/CompressionComponents/Oodle/FirstRun.bat
+++ b/Engine/Plugins/Runtime/PacketHandlers/CompressionComponents/Oodle/FirstRun.bat
@@ -48,7 +48,7 @@ pause
 
 
 :EnableHandler
-set HandlerCommandletParms=-run=OodleHandlerComponent.OodleTrainerCommandlet enable
+set HandlerCommandletParms=-run=OodleTrainerCommandlet enable
 set FinalHandlerCmdLine=%GameName:"=% %HandlerCommandletParms% -forcelogflush
 
 
