Small Network Prediction readme.txt update
#rnx

diff --git a/Engine/Plugins/Runtime/NetworkPrediction/readme.txt b/Engine/Plugins/Runtime/NetworkPrediction/readme.txt
--- a/Engine/Plugins/Runtime/NetworkPrediction/readme.txt
+++ b/Engine/Plugins/Runtime/NetworkPrediction/readme.txt
@@ -29,6 +29,9 @@ Getting Started:
 -Both NetworkPrediction and NetworkPredictionExtras are disabled by default. You will need to enable them for your project (manually edit .uproject or do so through editor plugin screen).
 -Once NetworkPredictionExtras is loaded, /NetworkPredictionExtras/Content/TestMap.umap can be loaded (must enable 'Show Engine Content' AND 'Show Plugin Content' in content browser!).
 -The "MockNetworkSimulation" can be tested anywhere just by typing "mns.Spawn".
+-It is recommend you use "net.DisableBandwithThrottling 1" while using this plugin (or bump up your NetSpeeds in DefaultEngine.ini).
+	-The plugin is quite greedy on bandwidth currently. It utilize redundant sending of data to avoid issues relating to packet loss.
+	-We hope to optimize this down the road but in general NetworkPrediction will be willing to trade bandwidth to preserve correctness in the simulation.
  
 Code to look at:
 -MockNetworkSimulation.h: Good place to start to see a bare minimum example of a predicted network simulation. Both the simulation and actor component live in this header.
@@ -106,6 +109,12 @@ Smoothing: Taking the output of the simulation and applying an additional layer
 // Release notes
 // ----------------------------------------------------------------------------------------------------------
 
+Update (2-18-20)
+-Sorry for the lack of public updates. Development has been ongoing but mostly moved to a non public stream at the moment. The work being done revolves around:
+	-Unreal Insights. We are implementing an Insights tool for debugging Network Prediction. This will replace the canvas/visual logger debugging in the plugin. It is very powerful. Hopefully v.1 will make it into main branch soon.
+	-We are in the early phases of "new movement system" design. No ETA.
+	-We are investigating deeper integration with physics. No ETA.
+
 Update (12-13-19)
 -Only small updates (if any) for the rest of the year. System is going through code reviews as we plan for writing a new general movement system on top of Network Prediction.
 -Things are still liable to change but for the most part every major feature is in.
