Network Physics
-Delete the Async Network Prediction system
-The original game thread Network Prediction will remain

diff --git a/Engine/Plugins/Runtime/NetworkPrediction/Source/NetworkPrediction/Private/Async/NetworkPredictionAsyncModelDefRegistry.cpp b/Engine/Plugins/Runtime/NetworkPrediction/Source/NetworkPrediction/Private/Async/NetworkPredictionAsyncModelDefRegistry.cpp
deleted file mode 100644
diff --git a/Engine/Plugins/Runtime/NetworkPrediction/Source/NetworkPrediction/Private/Async/NetworkPredictionAsyncProxy.cpp b/Engine/Plugins/Runtime/NetworkPrediction/Source/NetworkPrediction/Private/Async/NetworkPredictionAsyncProxy.cpp
deleted file mode 100644
diff --git a/Engine/Plugins/Runtime/NetworkPrediction/Source/NetworkPrediction/Private/Async/NetworkPredictionAsyncWorldManager.cpp b/Engine/Plugins/Runtime/NetworkPrediction/Source/NetworkPrediction/Private/Async/NetworkPredictionAsyncWorldManager.cpp
deleted file mode 100644
diff --git a/Engine/Plugins/Runtime/NetworkPrediction/Source/NetworkPrediction/Private/NetworkPhysics.cpp b/Engine/Plugins/Runtime/NetworkPrediction/Source/NetworkPrediction/Private/NetworkPhysics.cpp
deleted file mode 100644
diff --git a/Engine/Plugins/Runtime/NetworkPrediction/Source/NetworkPrediction/Private/NetworkPhysicsComponent.cpp b/Engine/Plugins/Runtime/NetworkPrediction/Source/NetworkPrediction/Private/NetworkPhysicsComponent.cpp
deleted file mode 100644
diff --git a/Engine/Plugins/Runtime/NetworkPrediction/Source/NetworkPrediction/Private/NetworkPredictionModule.cpp b/Engine/Plugins/Runtime/NetworkPrediction/Source/NetworkPrediction/Private/NetworkPredictionModule.cpp
--- a/Engine/Plugins/Runtime/NetworkPrediction/Source/NetworkPrediction/Private/NetworkPredictionModule.cpp
+++ b/Engine/Plugins/Runtime/NetworkPrediction/Source/NetworkPrediction/Private/NetworkPredictionModule.cpp
@@ -13,7 +13,6 @@
 #include "Misc/Parse.h"
 #include "String/ParseTokens.h"
 #include "NetworkPredictionModelDefRegistry.h"
-#include "Async/NetworkPredictionAsyncModelDefRegistry.h"
 
 #if WITH_EDITOR
 #include "Editor.h"
@@ -139,7 +138,6 @@ void FNetworkPredictionModule::FinalizeNetworkPredictionTypes()
 {
 	FGlobalCueTypeTable::Get().FinalizeCueTypes();
 	FNetworkPredictionModelDefRegistry::Get().FinalizeTypes();
-	UE_NP::FNetworkPredictionAsyncModelDefRegistry::Get().FinalizeTypes();
 }
 
 IMPLEMENT_MODULE( FNetworkPredictionModule, NetworkPrediction )
diff --git a/Engine/Plugins/Runtime/NetworkPrediction/Source/NetworkPrediction/Public/Async/NetworkPredictionAsyncDataStore.h b/Engine/Plugins/Runtime/NetworkPrediction/Source/NetworkPrediction/Public/Async/NetworkPredictionAsyncDataStore.h
deleted file mode 100644
diff --git a/Engine/Plugins/Runtime/NetworkPrediction/Source/NetworkPrediction/Public/Async/NetworkPredictionAsyncDefines.h b/Engine/Plugins/Runtime/NetworkPrediction/Source/NetworkPrediction/Public/Async/NetworkPredictionAsyncDefines.h
deleted file mode 100644
diff --git a/Engine/Plugins/Runtime/NetworkPrediction/Source/NetworkPrediction/Public/Async/NetworkPredictionAsyncGlobalTickContext.h b/Engine/Plugins/Runtime/NetworkPrediction/Source/NetworkPrediction/Public/Async/NetworkPredictionAsyncGlobalTickContext.h
deleted file mode 100644
diff --git a/Engine/Plugins/Runtime/NetworkPrediction/Source/NetworkPrediction/Public/Async/NetworkPredictionAsyncID.h b/Engine/Plugins/Runtime/NetworkPrediction/Source/NetworkPrediction/Public/Async/NetworkPredictionAsyncID.h
deleted file mode 100644
diff --git a/Engine/Plugins/Runtime/NetworkPrediction/Source/NetworkPrediction/Public/Async/NetworkPredictionAsyncInstanceData.h b/Engine/Plugins/Runtime/NetworkPrediction/Source/NetworkPrediction/Public/Async/NetworkPredictionAsyncInstanceData.h
deleted file mode 100644
diff --git a/Engine/Plugins/Runtime/NetworkPrediction/Source/NetworkPrediction/Public/Async/NetworkPredictionAsyncModelDef.h b/Engine/Plugins/Runtime/NetworkPrediction/Source/NetworkPrediction/Public/Async/NetworkPredictionAsyncModelDef.h
deleted file mode 100644
diff --git a/Engine/Plugins/Runtime/NetworkPrediction/Source/NetworkPrediction/Public/Async/NetworkPredictionAsyncModelDefRegistry.h b/Engine/Plugins/Runtime/NetworkPrediction/Source/NetworkPrediction/Public/Async/NetworkPredictionAsyncModelDefRegistry.h
deleted file mode 100644
diff --git a/Engine/Plugins/Runtime/NetworkPrediction/Source/NetworkPrediction/Public/Async/NetworkPredictionAsyncPersistentStorage.h b/Engine/Plugins/Runtime/NetworkPrediction/Source/NetworkPrediction/Public/Async/NetworkPredictionAsyncPersistentStorage.h
deleted file mode 100644
diff --git a/Engine/Plugins/Runtime/NetworkPrediction/Source/NetworkPrediction/Public/Async/NetworkPredictionAsyncProxy.h b/Engine/Plugins/Runtime/NetworkPrediction/Source/NetworkPrediction/Public/Async/NetworkPredictionAsyncProxy.h
deleted file mode 100644
diff --git a/Engine/Plugins/Runtime/NetworkPrediction/Source/NetworkPrediction/Public/Async/NetworkPredictionAsyncProxyImpl.h b/Engine/Plugins/Runtime/NetworkPrediction/Source/NetworkPrediction/Public/Async/NetworkPredictionAsyncProxyImpl.h
deleted file mode 100644
diff --git a/Engine/Plugins/Runtime/NetworkPrediction/Source/NetworkPrediction/Public/Async/NetworkPredictionAsyncSimCallback.h b/Engine/Plugins/Runtime/NetworkPrediction/Source/NetworkPrediction/Public/Async/NetworkPredictionAsyncSimCallback.h
deleted file mode 100644
diff --git a/Engine/Plugins/Runtime/NetworkPrediction/Source/NetworkPrediction/Public/Async/NetworkPredictionAsyncWorldManager.h b/Engine/Plugins/Runtime/NetworkPrediction/Source/NetworkPrediction/Public/Async/NetworkPredictionAsyncWorldManager.h
deleted file mode 100644
diff --git a/Engine/Plugins/Runtime/NetworkPrediction/Source/NetworkPrediction/Public/Async/Services/NetworkPredictionAsyncServiceRegistry.h b/Engine/Plugins/Runtime/NetworkPrediction/Source/NetworkPrediction/Public/Async/Services/NetworkPredictionAsyncServiceRegistry.h
deleted file mode 100644
diff --git a/Engine/Plugins/Runtime/NetworkPrediction/Source/NetworkPrediction/Public/Async/Services/NetworkPredictionAsyncService_Marshal.inl b/Engine/Plugins/Runtime/NetworkPrediction/Source/NetworkPrediction/Public/Async/Services/NetworkPredictionAsyncService_Marshal.inl
deleted file mode 100644
diff --git a/Engine/Plugins/Runtime/NetworkPrediction/Source/NetworkPrediction/Public/Async/Services/NetworkPredictionAsyncService_NetSerialize.inl b/Engine/Plugins/Runtime/NetworkPrediction/Source/NetworkPrediction/Public/Async/Services/NetworkPredictionAsyncService_NetSerialize.inl
deleted file mode 100644
diff --git a/Engine/Plugins/Runtime/NetworkPrediction/Source/NetworkPrediction/Public/Async/Services/NetworkPredictionAsyncService_Reconcile.inl b/Engine/Plugins/Runtime/NetworkPrediction/Source/NetworkPrediction/Public/Async/Services/NetworkPredictionAsyncService_Reconcile.inl
deleted file mode 100644
diff --git a/Engine/Plugins/Runtime/NetworkPrediction/Source/NetworkPrediction/Public/Async/Services/NetworkPredictionAsyncService_Registration.inl b/Engine/Plugins/Runtime/NetworkPrediction/Source/NetworkPrediction/Public/Async/Services/NetworkPredictionAsyncService_Registration.inl
deleted file mode 100644
diff --git a/Engine/Plugins/Runtime/NetworkPrediction/Source/NetworkPrediction/Public/Async/Services/NetworkPredictionAsyncService_Ticking.inl b/Engine/Plugins/Runtime/NetworkPrediction/Source/NetworkPrediction/Public/Async/Services/NetworkPredictionAsyncService_Ticking.inl
deleted file mode 100644
diff --git a/Engine/Plugins/Runtime/NetworkPrediction/Source/NetworkPrediction/Public/NetworkPhysics.h b/Engine/Plugins/Runtime/NetworkPrediction/Source/NetworkPrediction/Public/NetworkPhysics.h
deleted file mode 100644
diff --git a/Engine/Plugins/Runtime/NetworkPrediction/Source/NetworkPrediction/Public/NetworkPhysicsComponent.h b/Engine/Plugins/Runtime/NetworkPrediction/Source/NetworkPrediction/Public/NetworkPhysicsComponent.h
deleted file mode 100644
diff --git a/Engine/Plugins/Runtime/NetworkPrediction/Source/NetworkPrediction/Public/NetworkPredictionUtil.h b/Engine/Plugins/Runtime/NetworkPrediction/Source/NetworkPrediction/Public/NetworkPredictionUtil.h
--- a/Engine/Plugins/Runtime/NetworkPrediction/Source/NetworkPrediction/Public/NetworkPredictionUtil.h
+++ b/Engine/Plugins/Runtime/NetworkPrediction/Source/NetworkPrediction/Public/NetworkPredictionUtil.h
@@ -5,7 +5,40 @@
 #include "CoreTypes.h"
 #include "NetworkPredictionCheck.h"
 #include "UObject/NameTypes.h"
-#include "Async/NetworkPredictionAsyncDefines.h"
+
+namespace UE_NP
+{
+#ifndef NP_MAX_ASYNC_MODEL_DEFS
+#define NP_MAX_ASYNC_MODEL_DEFS 16
+#endif
+
+	const int32 MaxAsyncModelDefs = NP_MAX_ASYNC_MODEL_DEFS;
+
+#ifndef NP_NUM_FRAME_STORAGE
+#define NP_NUM_FRAME_STORAGE 64
+#endif
+
+	const int32 NumFramesStorage = NP_NUM_FRAME_STORAGE;
+
+#ifndef NP_FRAME_STORAGE_GROWTH
+#define NP_FRAME_STORAGE_GROWTH 8
+#endif
+
+	const int32 FrameStorageGrowth = NP_FRAME_STORAGE_GROWTH;
+
+#ifndef NP_FRAME_INPUTCMD_BUFFER_SIZE
+#define NP_FRAME_INPUTCMD_BUFFER_SIZE 16
+#endif
+
+	const int32 InputCmdBufferSize = NP_FRAME_INPUTCMD_BUFFER_SIZE;
+
+#ifndef NP_INLINE_SIMOBJ_INPUTS
+#define NP_INLINE_SIMOBJ_INPUTS 3
+#endif
+
+	const int32 InlineSimObjInputs = NP_INLINE_SIMOBJ_INPUTS;
+
+};
 
 // Sets index to value, resizing bit array if necessary and setting new bits to false
 template<typename BitArrayType>
diff --git a/Engine/Plugins/Runtime/NetworkPrediction/readme.txt b/Engine/Plugins/Runtime/NetworkPrediction/readme.txt
--- a/Engine/Plugins/Runtime/NetworkPrediction/readme.txt
+++ b/Engine/Plugins/Runtime/NetworkPrediction/readme.txt
@@ -2,193 +2,13 @@
 //	Network Prediction Plugin
 // --------------------------------------------------------------------------------------------------------------------
 
-10-22-2020:
+After a long road and many attempts, we are dropping support for the async version of Network Prediction. We felt the 
+complications it introduced into the physics system were too much to maintain and performance was still too poor in
+the worst/degenerate cases that it wasn't going to be a viable system for enough games to warrant the complexities.
 
-Sorry for lack of updates again. Development is never as smooth as we'd like. Most of these changes are from bugs/issues discovered
-while working on real world use cases of physics based simulations.
-
-New macro for tracing reasons why corrections happen: UE_NP_TRACE_RECONCILE
-* Use this in ShouldReconcile. When trace is enabled it will trace the given string when reconcile happens.
-*This string will appear in NP Insights and highlight the state the caused the reconcile.
-*(The UI is still not as intuitive as we'd like but this is a good start)
-
-Corrections caused by "local frame offset" changes are now traced
-*These are corrections that usually happen due to input starvation on the server: the client's InputCmd->ServerFrame relationship shifts.
-*This usually causes a corection as the clients local frame number shifts relative to the server frame numbers.
-*In NP Insights Sim Frame View, it wasn't obvious when this happened since in term of server frames, the client was still correct.
-*This is now called out in the "Sim Frame contents view" (bottom pane)
-
-Server-side Input starvation and overflow are now traced
-
-Added option to show number of buffered input commands in NP Insights
-
-Physics Received State now properly traced
-
-
-9-2-2020:
-
-Update for root motion. This improves on the previous implementation:
--UMockRootMotionSource as a base class for "root motion defining thing".
--New source types can be added without modifying the root motion simulation or driver code.
--Sources can have networked internal state and parameters.
--The root motion source is net serialized as a {SourceID, StartTime, Parameters} tuple via TMockRootMotionSourceProxy
--Sources define how their internal state/parameters are encoded into the Parameters blob via UMockRootMotionSource::SerializePayloadParameters
--Encoding SourceID is done through a new 
-
-This pattern is very relevenant for things like custom movement modes or other data-driven extensions of simulations.
-
-Minor updates:
--Trace local frame offset to Insights. 
--Store latest simulation time in FNetworkPredictionStateView
--Fix a couple issues where wrong time/frame was being traced, causing bugs in Insights
-
-
-
-8-10-2020:
-
-Initial mock root motion checked in. This is the begining of root motion networking which will find its way into the new movement system.
-Initial focus is just on having montage-based root motion support back in, plus non animation based sources like curves (more to come here).
-The animation team is planning for other animation based root motion sources.
-
-
-7-31-2020:
-
-Couple notable changes related to physics and UPrimitiveComponents
-
-The system now defaults to "UPrimitiveComponents are always in sync with their physics data when NP SimulationTick functions run".
-If you look at the previous version of FMockPhysicsSimulation::SimulationTick, we had to be very careful to interface directly
-with the underlying PhysicsActorHandle, rather than reading data off of any UPrimitiveComponents. This causes some pretty
-nasty anti-patterns, especially around scene queries.
-
-The cost is that we have to take a seperate pass during rollback for everyone to "RestoreFrame" prior to *anyone* resimulating
-a tick step. So, an extra pass through the registered instances and touching more memory than we did before.
-
-It would be possible to allow users to opt out of this: to say "I know what I'm doing and I am confident I can write all
-of my NP Code to interface directly with the physics engine". 
-
-We will see if this shows up in profiles as the system continues to mature and we build real stuff with it. For now we think
-its wiser to error on the side of being user friendly and less error prone.
-
-
-The specific changes here are:
--PhysicsActorHandle is no longer a side cart of data registered with the Driver/Simulation. 
--We now require FNetworkPredictionDriver<ModelDef>::GetPhysicsPrimitiveComponent() to get to the physics data.
--"RestoreFrame" is a driver-level function similiar to FinalizeFrame but is called only during resims where we want to push
-	the given sync/aux state to the physics scene (Whatever that means for you).	
-	The default implementation is provided for physics and doesn't need the user to implement anything.
-
-
-
-7-24-2020:
-
-Physics issues should be fixed. Still tracking down a few more bugs in the cue and interpolation systems.
-
-7-2-2020: Big Update
-
-This checkin is a large refactor of the NetworkPrediction system. We hope this is the last big set of changes, though
-there is still work to do. User code requires some fixups but should be a straight forward port overall. All existing 
-examples have been updated. See "Upgrade Notes" below.
-
-The primary motivation behind these changes are supporting Group Rollback and Physics. Group Rollback just meaning
-forward predicting multiple actors in step together.
-
-Group rollback and physics are supported only in fixed tick mode. Fix tick mode is now a global/system wide thing
-rather than a per-simulation setting. In other words, previously each simulation could choose its tick settings and
-would effectively be a black box in terms of when it decided to run a sim frame. Now, all fixed ticking happens
-at the system level: in UNetworkPredictionWorldManager.
-
-The system still supports what we are calling Independent Ticking: where a client-controlled simulation is ticked
-on the server at the same timesteps that the controlling client does (e.g, the client's local frame rate). Simulations
-running in independent mode cannot participate in group or physics rollback/resimulates. They are effectively on their
-own. This means that only the client-controlled simulation state is forward-predictable. In independent mode, you do 
-not forward predict other simulations.
-
-To restate this:
--Simulations can either tick A) independently or B) fixed.
--There is one fixed tick "group". There can be N independent ticking simulations (and are all completely independent).
--Independent only allows forward prediction of the client controlled simulation state (E.g, movement, GAS, etc).
--Fixed tick allows multiple actors, including non client-controlled actors to be forward predicted in step with the client.
--Fixed tick is the only mode that supports physics. More notes on physics below.
-
-Fix vs Independent ticking can be set at a few levels:
--The ModelDef itself can define what tick modes it supports via FNetworkPredictionDriver. "Capabilities"
--(By default, simulations that use Physics are only capable of fixed ticking)
--On spawned instances of a simulation, the authority can set fix/ticked (if both are supported). "Archetype".
--For cases where both modes are available, Project Settings -> Network Prediction (DefaultNetworkPrediction.ini) has 
-	default tick mode settings. (E.g, it will eventually fall back to this in most cases).
-
-
-Physics:
-
-*** Currently, this requires the engine to run in fixed tick mode itself. ***
-*** We are looking into expanding our support for fixed tick physics within a variable ticking engine. ***
-*** Engine Fixed tick mode will be enabled by default if you are using physics, via the bForceEngineFixTickForcePhysics 
-	property in Project Settings -> Network Prediction ***
-
-
--Physics support allows physics state to be recorded and resimulated by the physics engine (Chaos) itself.
--Chaos is being optimized to allow fast recording and resimulate steps.
--Physic-only sims are supported: e.g, an actor whose physics state is predicted but has no underlying NP state/tick.
--Physics + NP is also supported: e.g, a vehicle that has physics state and also a NP state/tick.
-
-
-There are additional changes around how the game code binds/interfaces with the NetworkPrediction system. The hope is
-that the new version is simplified and has less confusing boilerplate. Most of this is now done through the 
-FNetworkPredictionProxy struct. Refer to the example to see what has changed.
-
-
-Incomplete Features:
--Simulation Extrapolation has temporarily been removed. We support interpolation and forward prediction right now.
--Network Prediction Insights is due for a pass to better visualize these changes (mainly physics and group reconcile)
--Generalized server-side record/rewind (e.g, "lag compensation"). This will be an important part for independent ticking.
--Aux state is no longer stored sparsely, we intend to fix this.
--Some issues remain around independent ticking and cues. Working on fixing these.
-
-
-// -------------------------------------------------------------
-// Upgrade notes (7-2-2020)
-// -------------------------------------------------------------
-
-* ::Log function on user states has been changed to ::ToString(FAnsiStringBuilder&), you must now build ansi strings.
-	Careful not to use FVector::ToString functions that build TCHARS. This is an inconvenience but makes tracing 
-	state for Network Prediction Insights much more efficient. 
-
-* TNetworkSimBufferTypes renamed --> TNetworkPredictionStateTypes
-
-* FNetSimModelDefBase renamed --> FNetworkPredictionModelDef. There are additional changes to this:
-
-	* NP_MODEL_BODY (in header) and NP_REGISTER_MODELDEF macros are required for registering model defs
-
-	* Simulation, StateTypes (previously BufferTypes), and Driver are now defined in the ModelDef.
-
-	* ::GetName() and ::GetSortPriority() are also defined on the ModelDef.
-
-* The "Driver" is no longer a virtual interface that is inherited by the driving object (actor/component). The driver
-	class is just defined in the ModelDef and functions are called directly on it. This is done through FNetworkPredictionDriver.
-
-* ::InstantiateNetworkedSimulation renamed --> ::InitializeNetworkPredictionProxy 
-
-* Use NetworkPredictionProxy.Init<ModelDef>:: in InitializeNetworkPredictionProxy. You will need to manually include 
-	NetworkPredictionProxyInit.h in the file that does this (do not include that in the header of your actor/component).
-
-* FNetworkSimTime has been removed, we now use int32 for simulation time in the system. The wrapped structure was
-	not adding any real safety and was just getting in the way.
-
-* FinalizeFrame, ProduceInput now take pointers instead of references. This is to accommodate ModelDefs that have void
-	types (e.g, not AuxState. This allow for ::FinalizeFrame(const FSyncState*, const void*).
-
-* GetDebugName, GetVLogOwner, VisualLog are gone. Generic implementations of these for actors/components are now done in
-	FNetworkPredictionDriver. You can customize this behavior for your defs by specializing FNetworkPredictionDriver on your
-	ModelDef type.
-
-* TNetSimInput/TNetSimOutput now have pointers instead of references. Same reason as ProduceInput/FinalizeFrame: to support
-	void* cases.
-
-* TNetSimStateAccessor has been removed. You can now read/write to simulation state through the NetworkPredictionProxy
-	in a similar way. Note that you will need to explicitly include NetworkPredictionProxyWrite.h to any files that does this.
-	
-* You may need to include "Chaos" in your PublicDependencyModuleNames. Looking at a way to not require this but am not sure
-	it will be possible. 
+The original single threaded version of Network Prediction is preserved and unchanged. We still hope to use it to 
+build a new character moverment system with it. Physics support could come back into this version but it would be
+strictly opt in and only applicable to games with small number of objects and players.
 
 
 // -------------------------------------------------------------
@@ -213,8 +33,6 @@ the SimulationTick. Prediction and rollback support is generically provided for
 
 NetworkPredictionExtras is a supplementary plugin with sample content.
 
-NetworkPredictionInsights is a tool for debugging the Network Prediction system. Demo can be found here:
-https://www.youtube.com/watch?v=_rdt-v1nFlY
 
 // -------------------------------------------------------------
 //	Getting Started
@@ -229,15 +47,3 @@ are managed and coordinated.
 NetworkPredictionPhysicsComponent.h - Example of binding a physics-only sim. 
 
 MockPhysicsSimulation.h - Simple "controllable physics object" example. 
-
-
-// -------------------------------------------------------------
-//	Road Map
-// -------------------------------------------------------------
-
-Network Prediction is still a WIP. We don't foresee any more major changes but until it is officially released and 
-out of beta, there is a risk of API changes.
-
-The plan is to replace the current UCharacterMovementComponent and provide updates to GameplayAbilities so that all
-prediction is unified across these systems, plus is easier to extend/modify on a per project basis.
-
diff --git a/Engine/Plugins/Runtime/NetworkPredictionExtras/Source/NetworkPredictionExtras/Private/PhysicsEffect.cpp b/Engine/Plugins/Runtime/NetworkPredictionExtras/Source/NetworkPredictionExtras/Private/PhysicsEffect.cpp
deleted file mode 100644
diff --git a/Engine/Plugins/Runtime/NetworkPredictionExtras/Source/NetworkPredictionExtras/Private/PhysicsEffectsComponent.cpp b/Engine/Plugins/Runtime/NetworkPredictionExtras/Source/NetworkPredictionExtras/Private/PhysicsEffectsComponent.cpp
deleted file mode 100644
diff --git a/Engine/Plugins/Runtime/NetworkPredictionExtras/Source/NetworkPredictionExtras/Private/PhysicsMovement.cpp b/Engine/Plugins/Runtime/NetworkPredictionExtras/Source/NetworkPredictionExtras/Private/PhysicsMovement.cpp
deleted file mode 100644
diff --git a/Engine/Plugins/Runtime/NetworkPredictionExtras/Source/NetworkPredictionExtras/Private/PhysicsSimple.cpp b/Engine/Plugins/Runtime/NetworkPredictionExtras/Source/NetworkPredictionExtras/Private/PhysicsSimple.cpp
deleted file mode 100644
diff --git a/Engine/Plugins/Runtime/NetworkPredictionExtras/Source/NetworkPredictionExtras/Public/PhysicsEffect.h b/Engine/Plugins/Runtime/NetworkPredictionExtras/Source/NetworkPredictionExtras/Public/PhysicsEffect.h
deleted file mode 100644
diff --git a/Engine/Plugins/Runtime/NetworkPredictionExtras/Source/NetworkPredictionExtras/Public/PhysicsEffectTypes.h b/Engine/Plugins/Runtime/NetworkPredictionExtras/Source/NetworkPredictionExtras/Public/PhysicsEffectTypes.h
deleted file mode 100644
diff --git a/Engine/Plugins/Runtime/NetworkPredictionExtras/Source/NetworkPredictionExtras/Public/PhysicsEffectsComponent.h b/Engine/Plugins/Runtime/NetworkPredictionExtras/Source/NetworkPredictionExtras/Public/PhysicsEffectsComponent.h
deleted file mode 100644
diff --git a/Engine/Plugins/Runtime/NetworkPredictionExtras/Source/NetworkPredictionExtras/Public/PhysicsMovement.h b/Engine/Plugins/Runtime/NetworkPredictionExtras/Source/NetworkPredictionExtras/Public/PhysicsMovement.h
deleted file mode 100644
diff --git a/Engine/Plugins/Runtime/NetworkPredictionExtras/Source/NetworkPredictionExtras/Public/PhysicsSimple.h b/Engine/Plugins/Runtime/NetworkPredictionExtras/Source/NetworkPredictionExtras/Public/PhysicsSimple.h
deleted file mode 100644
