Add missing APL for Quest support in the core OpenXR plugin
This is necessary when the OculusVR plugin is disabled

diff --git a/Engine/Plugins/Runtime/OpenXR/OculusOpenXRLoader_APL.xml b/Engine/Plugins/Runtime/OpenXR/OculusOpenXRLoader_APL.xml
new file mode 100644
--- /dev/null
+++ b/Engine/Plugins/Runtime/OpenXR/OculusOpenXRLoader_APL.xml
@@ -0,0 +1,102 @@
+<?xml version="1.0" encoding="utf-8"?>
+<!--LibOVRPlatform plugin additions-->
+<root xmlns:android="http://schemas.android.com/apk/res/android">
+	<!-- init section is always evaluated once per architecture -->
+	<init>
+		<log text="OculusOpenXRLoader"/>
+		<setBool result="bSupported" value="false"/>
+		<isArch arch="armeabi-v7a">
+			<setBool result="bSupported" value="true"/>
+		</isArch>
+		<isArch arch="arm64-v8a">
+			<setBool result="bSupported" value="true"/>
+		</isArch>
+		<setBoolFromPropertyContains result="bPackageForOculusQuest" ini="Engine" section="/Script/AndroidRuntimeSettings.AndroidRuntimeSettings" property="PackageForOculusMobile" contains="Quest"/>
+		<setBoolFromPropertyContains result="bPackageForOculusQuest2" ini="Engine" section="/Script/AndroidRuntimeSettings.AndroidRuntimeSettings" property="PackageForOculusMobile" contains="Quest2"/>
+		<setBoolOr result="bPackageForOculusMobile" arg1="$B(bPackageForOculusQuest)" arg2="$B(bPackageForOculusQuest2)"/>
+		<setBoolAnd result="bPackageForOculusQuestAndQuest2" arg1="$B(bPackageForOculusQuest)" arg2="$B(bPackageForOculusQuest2)"/>
+	</init>
+
+	<!-- optional updates applied to AndroidManifest.xml -->
+	<androidManifestUpdates>
+		<if condition="bPackageForOculusMobile">
+			<true>
+				<addFeature android:name="android.hardware.vr.headtracking" android:version="1" android:required="true" />
+
+				<dumpvars/>
+				<!-- Add Activity Specific Flags -->
+				<loopElements tag="activity">
+					<setStringFromAttribute result="activityName" tag="$" name="android:name"/>
+					<setBoolIsEqual result="bGameActivity" arg1="$S(activityName)" arg2="com.epicgames.unreal.GameActivity"/>
+					<if condition="bGameActivity">
+						<true>
+							<!-- Add VR Intent Filter, Permissions, and Features -->
+							<if condition="bPackageForOculusMobile">
+								<true>
+									<setElement result="vrIntent" value="category"/>
+									<addAttribute tag="$vrIntent" name="android:name" value="com.oculus.intent.category.VR"/>
+									<addElement tag="intent-filter" name="vrIntent"/>      
+									<addPermission android:name="android.permission.READ_EXTERNAL_STORAGE"/>
+								</true> 
+
+								<!-- Add Focus Aware Flag -->
+								<setElement result="focusAware" value="meta-data" />
+								<addAttribute tag="$focusAware" name="android:name" value="com.oculus.vr.focusaware" />
+								<addAttribute tag="$focusAware" name="android:value" value="true" />
+								<addElement tag="activity" name="focusAware"/>
+							</if>
+						</true>
+					</if>
+				</loopElements>
+
+				<!-- Add Application Specific Flags -->
+				<loopElements tag="application">
+					<!-- Add SupportedDevices Tag -->
+					<setElement result ="supportedDevices" value="meta-data" />
+					<addAttribute tag="$supportedDevices" name="android:name" value="com.oculus.supportedDevices" />
+					<if condition="bPackageForOculusQuestAndQuest2">
+						<true>
+							<addAttribute tag="$supportedDevices" name="android:value" value="quest|delmar" />
+						</true>
+						<false>
+							<if condition="bPackageForOculusQuest2">
+								<true>
+									<addAttribute tag="$supportedDevices" name="android:value" value="delmar" />
+								</true>
+								<false>
+									<addAttribute tag="$supportedDevices" name="android:value" value="quest" />
+								</false>
+							</if>
+						</false>
+					</if>
+					<addElement tag="application" name="supportedDevices"/>
+
+				</loopElements>
+			</true>
+		</if>
+	</androidManifestUpdates>
+
+	<!-- optional files or directories to copy to Intermediate/Android/APK -->
+	<resourceCopies>
+		<if condition="bPackageForOculusMobile">
+			<true>
+				<isArch arch="arm64-v8a">
+					<log text="Copying libopenxr_loader.so"/>
+					<copyFile src="$S(EngineDir)/Source/ThirdParty/Oculus/OculusOpenXRLoader/OculusOpenXRLoader/Lib/arm64-v8a/libopenxr_loader.so"
+								dst="$S(BuildDir)/libs/arm64-v8a/libopenxr_loader.so" />
+				</isArch>
+			</true>
+		</if>
+	</resourceCopies>
+
+
+	<!-- optional libraries to load in GameActivity.java before libUnreal.so -->
+	<soLoadLibrary>
+		<!-- need this if plugin enabled and supported architecture, even if not packaged for LibOVRPlatform -->
+		<if condition="bPackageForOculusMobile">
+			<true>
+				<loadLibrary name="openxr_loader" failmsg="openxr_loader library not loaded and may be required for Oculus VR." />
+			</true>
+		</if>
+	</soLoadLibrary>
+</root>
diff --git a/Engine/Plugins/Runtime/OpenXR/Source/OpenXRHMD/OpenXRHMD.Build.cs b/Engine/Plugins/Runtime/OpenXR/Source/OpenXRHMD/OpenXRHMD.Build.cs
--- a/Engine/Plugins/Runtime/OpenXR/Source/OpenXRHMD/OpenXRHMD.Build.cs
+++ b/Engine/Plugins/Runtime/OpenXR/Source/OpenXRHMD/OpenXRHMD.Build.cs
@@ -1,7 +1,7 @@
 // Copyright Epic Games, Inc. All Rights Reserved.
 
-using UnrealBuildTool;
 using System.IO;
+using System.Linq;
 
 namespace UnrealBuildTool.Rules
 {
@@ -107,7 +107,14 @@ namespace UnrealBuildTool.Rules
                 });
 
                 AddEngineThirdPartyPrivateStaticDependencies(Target, "Vulkan");
-            }
+			}
+
+			if (Target.Platform == UnrealTargetPlatform.Android && !Target.EnablePlugins.Contains("OculusVR"))
+			{
+				// If the Oculus plugin is not enabled we need to include our own APL
+				string PluginPath = Utils.MakePathRelativeTo(ModuleDirectory, Target.RelativeEnginePath);
+				AdditionalPropertiesForReceipt.Add("AndroidPlugin", Path.Combine(PluginPath, "OculusOpenXRLoader_APL.xml"));
+			}
 		}
 	}
 }
