First pass at Gameplay Ability System plug-in documentation markdown.  This documentation is meant to be easily accessible for the community to give feedback on areas where our documentation is lacking.  We hope to get pull-requests and merge this documentation with the official documentation on the EDC website.

diff --git a/Engine/Plugins/Runtime/GameplayAbilities/README.md b/Engine/Plugins/Runtime/GameplayAbilities/README.md
new file mode 100644
--- /dev/null
+++ b/Engine/Plugins/Runtime/GameplayAbilities/README.md
@@ -0,0 +1,174 @@
+# Purpose of this Documentation
+
+This documentation is meant to support and enhance the [official Gameplay Ability System Unreal Developer Community documentation](https://dev.epicgames.com/documentation/en-us/unreal-engine/gameplay-ability-system-for-unreal-engine).  In particular, this document lives in the code repository under [Gameplay Ability System plug-in folder](./) and thus any user reading this documentation can submit a pull request to clarify functionality, update inaccurate information, or work with the community to flesh out areas that are missing.
+
+It is worth noting there are extensive resources that the wider end-user developer community has written.  One such source of knowledge is the [tranek GAS documentation]([REDACTED_PROFILE_URL]/GASDocumentation) which is highly detailed and an excellent resource for implementation details, but risks falling out of date with new feature additions or changes.
+
+# Overview of the Gameplay Ability System
+
+The Gameplay Ability System is a framework for building abilities and interactions that Actors can own and trigger. This system is designed mainly for RPGs, action-adventure games, MOBAs, and other types of games where characters have abilities that need to coordinate mechanics, visual effects, animations, sounds, and data-driven elements, although it can be adapted to a wide variety of projects. The Gameplay Ability System also supports replication for multiplayer games, and can save developers a lot of time scaling up their designs to support multiplayer.
+
+The concepts that the Gameplay Ability System uses are:
+
+- [Gameplay Attributes](#gameplay-attributes):  An enhancement to float properties that allow them to be temporarily modified (buffed) and used in complex calculations such as damage.
+- [Gameplay Tags](#gameplay-tags):  A hierarchical naming system that allows you to specify states of Actors, and properties of Assets.  A powerful query system allows designers to craft logic statements around these.
+- [Gameplay Cues](#gameplay-cues):  A visual and audio effects system based on Gameplay Tags which allow decoupling of the FX and the implementation.
+- [Gameplay Abilities](#gameplay-abilities):  The code that actually triggers when an action is performed.  Typically a Blueprint graph.
+- [Gameplay Effects](#gameplay-effects):  Predefined rulesets about how to apply all of the above.
+
+One of the designers' often mentioned goals of the Gameplay Ability System is to maintain a record of who triggered a complex set of interactions, so that we can keep proper account of which Actor did what.  For instance, if a Player activates a [Gameplay Ability](#gameplay-abilities) that spawns a poison cloud (possibly represented with a [Gameplay Cue](#gameplay-cue)) that then does damage-over-time using a [Gameplay Effect](#gameplay-effects) which eventually reduces an Actor's Health [Gameplay Attribute](#gameplay-attributes) to zero, we should be able to award the kill to the initiating Player.
+
+It is worth mentioning the damage system upfront, as it's a pervasive example throughout the documentation.  You may be familiar with the [AActor::TakeDamage function](https://www.unrealengine.com/blog/damage-in-ue4) which was used for many years.  At Epic, we no longer use that system internally; instead all damage is done through the Gameplay Ability System.  By using the Gameplay Ability System, we allow buffs/debuffs and an extensive and ever-changing list of damage types based on Gameplay Tags.  You can look at [Lyra](https://dev.epicgames.com/documentation/unreal-engine/lyra-sample-game-in-unreal-engine) as an example that uses the Gameplay Ability System extensively, as well as a rudimentary damage system.
+
+---
+
+# Ability System Component {#asc}
+
+The Ability System Component (commonly abbreviated ASC) is the top-level ActorComponent that you use to interface with the Gameplay Ability System (commonly abbreviated GAS).  It is a monolithic class that encapsulates almost all of the functionality GAS uses.  By funneling all of the functionality through the ASC, we are able to better encapsulate and enforce the rules about activation, replication, prediction, and side-effects.
+
+While the Ability System Component *is* an ActorComponent, we typically recommend against putting it on a Player's Pawn.  Instead, for a Player, it should be on the PlayerState.  The reason for this is that the Pawn is typically destroyed upon death in multiplayer games, and GAS typically has functionality (be it [Gameplay Attributes](#gameplay-attributes), [Gameplay Tags](#gameplay-tags), or [Gameplay Abilities](#gameplay-abilities)) that should persist beyond death.  For non-player AI-driven characters (e.g. AI that are not bots), it is suitable to put the ASC on the Pawn because it needs to replicate data to [Simulated Proxies](https://dev.epicgames.com/documentation/en-us/unreal-engine/actor-role-and-remote-role-in-unreal-engine#actorrolestates).
+
+---
+
+# Gameplay Attributes {#gameplay-attributes}
+
+Gameplay Attributes (often just referred to as simply *Attributes*) are essentially *float* properties that are wrapped in a FGameplayAttributeData instance.  The reason for doing so is to allow for a *BaseValue* which one can think of as an unaltered intrinsic value of the Actor, and a *CurrentValue* which one can think of as the value that currently applies, after all of the buffs and debuffs of the Actor are taken into account.  These *Attributes* must live in an [AttributeSet](#attribute-sets).  There is Editor tooling around the use of *Attributes* that allow them to be selected and used inside [Gameplay Effects](#gameplay-effects) (and others) to ensure buffs and debuffs work correctly.
+
+Attributes are often replicated, thus keeping the client in sync with the server values, but that does not *always* need to be the case.  For instance, certain *meta-Attributes* can be used to store temporary data used for calculations, allowing these intermediate results to have full buff/debuff aggregation capabilities; these *meta-Attributes* are not replicated because they are typically reset after a calculation.
+
+Since Gameplay Attributes are easily accessible through native or Blueprint code, it's tempting to modify them directly.  However, the Gameplay Ability System is designed such that all modifications to the Attributes should be done through [Gameplay Effects](#gameplay-effects) to ensure they can be network predicted and rolled-back gracefully.
+
+[Developer Community Gameplay Attribute & AttributeSet Docs](https://dev.epicgames.com/documentation/unreal-engine/gameplay-attributes-and-attribute-sets-for-the-gameplay-ability-system-in-unreal-engine)
+
+## AttributeSets {#attribute-sets}
+
+AttributeSets are simply classes derived from [UAttributeSet class](./Source/GameplayAbilities/Public/AttributeSet.h).  The AttributeSets typically contain multiple Gameplay Attributes that encompass all properties for a specific game feature (such as a jetpack item, but the most commonly cited example is the damage system).  AttributeSets must be added to the [Ability System Component](#asc) by the server.  AttributeSets are typically replicated to the client, but not all Attributes are replicated to the client (they are configured on a per-Attribute basis).
+
+## Attribute Modifiers {#attribute-modifiers}
+
+Attribute Modifiers are how we buff and debuff Attributes.  These are setup through the [Gameplay Effects'](#gameplay-effects) `Modifiers` property.  Once a modifier is 'active', it is stored in the [Ability System Component](#asc) and all requests for the value go through a process called *aggregation*.
+
+The rules for *aggregation* can be unexpected to a new user.  For instance, if there are multiple values that modify a single attribute, the modifiers are added together before the result is computed.  Let's take an example of a multiplier of 10% added to damage, and another multiplier of 30% added to damage.  If one were purely looking at the numbers, one could think `Damage * 1.1 * 1.3 = 1.43` thus damage would be modified by *43%*.  However, the system takes these modifier operators into account and adds them separately before performing the final multiplier calculation, giving an expected result of `10% + 30% = 40%`.
+
+---
+
+# Gameplay Abilities {#gameplay-abilities}
+
+Gameplay Abilities are derived from the [UGameplayAbility class](./Source/GameplayAbilities/Public/GameplayAbility.h).  They define what an in-game ability does, what (if anything) it costs to use, when or under what conditions it can be used, and so on.  Because Gameplay Abilities are implemented in native or Blueprints, it can do anything a Blueprint graph can do.  Unlike traditional Blueprints, they are capable of existing as instanced objects running asynchronously -- so you can run specialized, multi-stage tasks (called [Gameplay Ability Tasks](#gameplay-ability-tasks).  Examples of Gameplay Abilities would be a character dash, or an attack.
+
+Think of a Gameplay Ability as the bundle of functions that correspond to the action you're performing.  There are complex rules about who can activate them, how they activate, and how they are predicted (locally executed ahead of the server acknowledgement).  You trigger them through the Ability System Component (typically through a TryActivate function).  But they can also be triggered through complex interactions (if desired) such as through Gameplay Events, [Gameplay Tags](#gameplay-tags), [Gameplay Effects](#gameplay-effects), and Input (which the [ASC](#asc) handles internally).
+
+[Developer Community Gameplay Abilities docs](https://dev.epicgames.com/documentation/unreal-engine/using-gameplay-abilities-in-unreal-engine)
+
+## Gameplay Ability Tasks {#gameplay-ability-tasks}
+Gameplay Abilities often make use of [Gameplay Ability Tasks](https://dev.epicgames.com/documentation/unreal-engine/gameplay-ability-tasks-in-unreal-engine).  Gameplay Ability Tasks are latent Blueprint nodes that allow your Gameplay Ability to 'pause' for the frame while it awaits some event.  They can also perform network functionality which hide complex implementation details from the Blueprint designer.
+
+---
+
+# Gameplay Effects {#gameplay-effects}
+
+The purpose of Gameplay Effects is to modify an Actor in a predictable (and undoable) way.  Think of the verb Affect when you think of Gameplay Effects.  These are not Visual Effects or Sound Effects (those are called Gameplay Cues).  The Gameplay Effects are *applied* using [Gameplay Effect Specs](#gameplay-effect-specs) through the Ability System Component.
+
+- Gameplay Effects that have a Duration (non-Instant) will automatically undo any modifications to the Actor upon removal. Instant ones will modify the Attribute's *BaseValue*.
+- These are typically data-only Blueprints, though native implementations are also supported.
+- They should be static after compile time; there is no way to modify them during runtime (Gameplay Effect Specs are the runtime version).
+- They are essentially a complex datatable of things that should occur to a Target Actor when 'applied'.
+- Composed from these pieces:
+	- Duration / Timing data (such as how long the Effect lasts for, or how it periodically executes).
+	- Rules for Stacking the Gameplay Effects.
+	- Attribute Modifiers (data that controls how a Gameplay Attribute is modified and thus can be undone).
+	- Custom Executions (a user definable function that executes every time a Gameplay Effect is applied).
+	- Gameplay Effect Components (fragments of code / behavior to execute when applied).
+	- Rules for applying Gameplay Cues (the VisualFX and AudioFX).
+		
+## Gameplay Effect Components {#gameplay-effect-components}
+
+Gameplay Effect Components are introduced in UE5.3 to declutter the Gameplay Effect user interface and allow users of the Engine to provide their own game-specific functionality to Gameplay Effects.  
+
+Read the interface for [UGameplayEffectComponent](./Source/GameplayAbilities/Public/GameplayEffectComponent.h)
+
+## Gameplay Effect Specs {#gameplay-effect-specs}
+
+These are the runtime wrapper structs for a Gameplay Effect.  They define the Gameplay Effect, any dynamic parameters (such as SetByCaller data), and the tags as they existed when the Spec was created.  The majority of the runtime API's use a *GameplayEffectSpec* rather than a *GameplayEffect*.
+
+## Gameplay Effect Executions {#gameplay-effect-executions}
+
+Gameplay Effect Executions are game-specific, user-written functions that are configured to execute when particular Gameplay Effects execute.  They typically read from and write to [Gameplay Attributes](#gameplay-attributes).  These are used when the calculations are much more complex than can be achieved with a simple attribute modifier.  Examples of this would be a damage system (see the [Lyra Example](https://dev.epicgames.com/documentation/unreal-engine/lyra-sample-game-in-unreal-engine)).
+
+---
+
+# Gameplay Tags {#gameplay-tags}
+
+The Gameplay Ability System uses Gameplay Tags extensively throughout.  See the [official Developer Community documentation](https://dev.epicgames.com/documentation/unreal-engine/using-gameplay-tags-in-unreal-engine) for more details.
+
+---
+
+# Gameplay Cues (#gameplay-cues)
+
+Gameplay Cues are a system for decoupling visual and audio fx from gameplay code.  On start-up, special Gameplay Cue asset folders are scanned for [Gameplay Cue Sets](./Source/GameplayAbilities/Public/GameplayCueSet.h), and *Gameplay Cue Notify* classes.
+
+The implementer of a gameplay feature will either call the [Ability System Component](#asc)'s GameplayCue functions, or the [GameplayCueManager](./Source/GameplayAbilities/Public/GameplayCueManager.h)'s Gameplay Cue functions with a specific [Gameplay Tag](#gameplay-tag).  The effects artist will then create a *Gameplay Cue Notify* that corresponds to that tag.  The Gameplay Cue Manager is responsible for routing that specific tag to the proper *Gameplay Cue Notify*.
+
+## Gameplay Cue replication
+
+The details of Gameplay Cue replication are complex and worth noting.  Because these are cosmetic-only, there are unreliable RPC's that are used to communicate the execution of short-lived "*Burst*" cues.  We also use variable replication to synchronize the existance of longer cues (typically referred to as *Looping* or *Actor Notfies*).  This two-tiered approach ensures that Gameplay Cues can be dropped as unimportant, but also ensures important cues can be visible to any clients that *become relevant* according to the network system.
+
+Due to these Gameplay Cues needing to obey network relevancy (i.e. far away players should not replicate their Cues, but newly relevant ones should) and the fact that the PlayerState is *always relevant*, there is a *replication proxy* system.  The Player's Pawn (who has its [ASC](#asc) on the PlayerState) should implement the [IAbilitySystemReplicationProxyInterface](./Source/GameplayAbilities/Public/AbilitySystemReplicationProxyInterface.h).  When turning on the ASC's ReplicationProxyEnabled variable, all unreliable Gameplay Cue RPC's will go through the proxy interface (the Pawn, which properly represents relevancy).
+
+An advanced form of replication proxies also exists for the property replication so it may follow the same relevancy rules.  See `FMinimalGameplayCueReplicationProxy` in the [GameplayCueInterface](./Source/GameplayAbilities/Public/GameplayCueInterface.h).
+
+---
+
+# How Gameplay Prediction Works
+
+There is documentation for how the Gameplay Prediction mechanisms work at the top of [GameplayPrediction.h](./Source/GameplayAbilities/Public/GameplayPrediction.h).
+
+---
+
+# Debugging the Gameplay Ability System
+
+## Legacy ShowDebug Functionality
+
+Prior to UE5.4, the way to debug the Gameplay Ability System was to use the "ShowDebug AbilitySystem" command.  Once there, you can cycle through the categories using the command `AbilitySystem.Debug.NextCategory` or explicitly choose a category using `AbilitySystem.Debug.SetCategory`.  This system is no longer maintained and may be deprecated in future versions.  You should instead be looking at the [Gameplay Debugger](#gameplay-debugger) functionality.
+
+## Gameplay Debugger {#gameplay-debugger}
+
+New in UE5.4, there is enhanced Gameplay Debugger functionality for the Gameplay Ability System.  This functionality is preferred over the ShowDebug system and should be your first line of defense in debugging GAS.  To enable it, open the Gameplay Debugger typically by using `shift-apostrophe` (`shift-'`) to select the locally controlled player, or simply the `apostrophe` (`'`) key to select the Actor that is closest to your reticule.
+
+The debugger will show you the AbilitySystemComponent's current state as it pertains to Gameplay Tags, Gameplay Abilities, Gameplay Effects, and Gameplay Attributes.  In a networked game, the color coding helps to differentiate between how the server and client view their state.
+
+## Console Commands
+
+There are console commands that help both in developing and debugging GAS.  They are a great way to verify that your assumptions are correct about how abilities and effects should be activated, and coupled with the [Gameplay Debugger](#gameplay-debugger), what your state should be once executed.
+
+All Ability System debug commands are prefixed with `AbilitySystem`.  The functionality we're reviewing here exists in the [AbilitySystemCheatManagerExtension](./Source/GameplayAbilities/Private/AbilitySystemCheatManagerExtension.cpp).  The source code also serves as an excellent reference  how to properly trigger the Gameplay Abilities and Gameplay Effects in native code (and what the expected results would be, depending on their configurations).
+
+By implementing these in a Cheat Manager Extension, you are able to properly execute them as a local player, or on the server.  Many of the commands allow such a distinction with the `-Server` argument (read the command documentation or source code for more information).
+
+One of the gotchas when using these commands is that the assets should be loaded prior to their use.  This is easily done in the Editor by simply right-clicking on the assets you want to use and clicking "Load Assets".
+
+`AbilitySystem.Ability.Grant <ClassName/AssetName>` Grants an Ability to the Player.  Granting only happens on the Authority, so this command will be sent to the Server.
+`AbilitySystem.Ability.Activate [-Server] <TagName/ClassName/AssetName>` Activate a Gameplay Ability.  Substring name matching works for Activation Tags (on already granted abilities), Asset Paths (on non-granted abilities), or Class Names on both.  Some Abilities can only be activated by the Client or the Server and you can control all of these activation options by specifying or ommitting the `-Server` argument.
+`AbilitySystem.Ability.Cancel [-Server] <PartialName>` Cancels (prematurely Ends) a currently executing Gameplay Ability.  Cancelation can be initiated by either the Client or Server.
+`AbilitySystem.Ability.ListGranted` List the Gameplay Abilities that are granted to the local player.  Since granting is done on the Server but replicated to the Client, these should always be in sync (so no option for -Server).
+
+`AbilitySystem.Effect.ListActive [-Server]` Lists all of the Gameplay Effects currently active on the Player.
+`AbilitySystem.Effect.Remove [-Server] <Handle/Name>` Remove a Gameplay Effect that is currently active on the Player.
+`AbilitySystem.Effect.Apply [-Server] <Class/Assetname> [Level]` Apply a Gameplay Effect on the Player.  Substring name matching works for Asset Tags, Asset Paths, or Class Names.  Use -Server to send to the server (default is apply locally).
+
+Gameplay Cues have their own set of debug commands.
+
+`AbilitySystem.LogGameplayCueActorSpawning` Log when we create GameplayCueNotify_Actors.
+`AbilitySystem.DisplayGameplayCues` Display GameplayCue events in world as text.
+`AbilitySystem.GameplayCue.DisplayDuration` Configure the amount of time Gameplay Cues are drawn when `DisplayGameplayCues` is enabled.
+`AbilitySystem.DisableGameplayCues` Disables all GameplayCue events in the world.
+`AbilitySystem.GameplayCue.RunOnDedicatedServer` Run gameplay cue events on dedicated server.
+`AbilitySystem.GameplayCueActorRecycle` Allow recycling of GameplayCue Actors.
+`AbilitySystem.GameplayCueActorRecycleDebug` Prints logs for GC actor recycling debugging.
+`AbilitySystem.GameplayCueCheckForTooManyRPCs` Warns if gameplay cues are being throttled by network code.
+
+## Visual Logger
+
+New in UE5.4, there has been extra care put into the [Visual Logger](https://dev.epicgames.com/documentation/en-us/unreal-engine/visual-logger-in-unreal-engine) facilities for the Gameplay Ability System.  The [Visual Logger](https://dev.epicgames.com/documentation/en-us/unreal-engine/visual-logger-in-unreal-engine) is useful to see the complex interactions of Gameplay Abilities and Gameplay Effects over time.  The Visual Logger always captures the verbose logs and saves a snapshot of the state of the Ability System Component on every frame there is a log entry.
+
+In UE5.4, the [Visual Logger](https://dev.epicgames.com/documentation/en-us/unreal-engine/visual-logger-in-unreal-engine) now correctly orders the events between clients and servers when using Play In Editor.  This makes the Visual Logger especially useful for debugging how the client and server interact when activating abilities, gameplay effects, and modifying attributes.
\ No newline at end of file
