Deprecating the use of Replicated Properties in Gameplay Abilities.

The reason these are deprecated is that since properties replication uses a separate channel than reliable RPC's, the ordering of variable replication with respect to all other Gameplay Ability functionality (such as activating a GA, or providing TargetData) is not guaranteed.  This can lead to very subtle and unexpected bugs.

diff --git a/Engine/Plugins/Runtime/GameplayAbilities/README.md b/Engine/Plugins/Runtime/GameplayAbilities/README.md
--- a/Engine/Plugins/Runtime/GameplayAbilities/README.md
+++ b/Engine/Plugins/Runtime/GameplayAbilities/README.md
@@ -61,8 +61,89 @@ Think of a Gameplay Ability as the bundle of functions that correspond to the ac
 [Developer Community Gameplay Abilities docs](https://dev.epicgames.com/documentation/unreal-engine/using-gameplay-abilities-in-unreal-engine)
 
 ## Gameplay Ability Tasks {#gameplay-ability-tasks}
+
 Gameplay Abilities often make use of [Gameplay Ability Tasks](https://dev.epicgames.com/documentation/unreal-engine/gameplay-ability-tasks-in-unreal-engine).  Gameplay Ability Tasks are latent Blueprint nodes that allow your Gameplay Ability to 'pause' for the frame while it awaits some event.  They can also perform network functionality which hide complex implementation details from the Blueprint designer.
 
+## Gameplay Ability Specs
+
+The Gameplay Ability Specs are runtime-defined data which augment and tie together parameters used for the Gameplay Ability.  It serves two purposes:
+
+1. To configure the Gameplay Ability parameters prior to giving/granting the ability.  For instance, it defines what 'Level' of the ability you are granting.
+2. To store information about the granted Gameplay Ability that is shared between all instances of the Gameplay Ability.
+
+A lot of the [Ability System Component](#asc)'s interface deals with Gameplay Ability Specs, or after being granted, typically a Gameplay Ability Spec Handle.  The Handles are a way to succinctly refer to a Gameplay Ability Spec in both native and Blueprint code without worrying about the dangers of holding onto a pointer (and having that pointer be reallocated).  Whenever you want to refer to an already-granted Ability, use an Ability Spec Handle instead.
+
+The corresponding class for [FGameplayAbilitySpec](./Source/GameplayAbilities/Public/GameplayAbilitySpec.h) and [FGameplayAbilitySpecHandle](./Source/GameplayAbilities/Public/GameplayAbilitySpecHandle.h).
+
+## Gameplay Ability Instancing Policy
+
+The instancing policy determines when a Gameplay Ability is instanced (the Gameplay Ability object is created) and thus controls the lifetime of the GA.  The safest, and most feature-supported choice is InstancedPerActor.
+
+### InstancedPerActor
+
+When choosing InstancedPerActor, the Gameplay Ability will be instanced when its corresponding Gameplay Ability Spec is first given (granted) to the Actor.  The instance lives until the Gameplay Ability Spec is removed from the Actor.  This lifetime mimics what most users expect:  You are granted an ability and immediately have an instance of it.
+
+The lifetime semantics come with some pitfalls you should be aware of:
+
+- Since the ability continues to exist after it has ended, none of the variables will be reset for next activation.  Thus it's the user's responsibility to reset the variables to the defaults in EndAbility.
+- Prior to UE5.5, you could receive function calls such as OnGiveAbility/OnRemoveAbility on the _instance_ immediately, before the ability had ever been activated.  This isn't true of the other instancing types, which execute said functions on the CDO.  UE5.5 deprecates such functions in favor of explicit execution on the CDO.
+- There is a function you may see often called GetPrimaryInstance.  The Primary Instance refers to the InstancedPerActor's one-and-only instance; it does not apply to other instancing types.
+
+### InstancedPerExecution
+
+When choosing InstancedPerExecution, you receive a new instance of the Gameplay Ability for each and every activation.  Some things you should be aware of:
+
+- The instancing happens on activation (not prior to it).  It is possible to Grant & Revoke an InstancedPerExecution ability without ever instancing it.
+- Replicated Gameplay Abilities (GA's which contain RPC's or Replicated Variables) are relatively expensive, as a new GA must be sent for every activation.
+- Unlike InstancedPerActor, an individual instance is always active (otherwise it would have not been created).  It is garbage collected immediately upon ending.
+
+### NonInstanced (Deprecated)
+
+Prior to UE5.5, we had functionality for Non-Instanced Gameplay Abilities.  Since these Gameplay Abilities were never instanced, they could not be replicated or even hold state (e.g. contain variables).  All functions were called on the ClassDefaultObject and thus all state had to be held on the Gameplay Ability Spec.  This made them very confusing to use.  The same functionality can be achieved by simply using InstancedPerActor and never revoking it; the cost is just a single allocation (instance) of a UGameplayAbility.
+
+## Replication
+
+There is a replication policy variable on the Gameplay Abilities.  The setting controls whether or not you are able to use Remote Procedure Calls (RPC's) or Replicated Variables (now deprecated, see below).  It does *not* control if a Gameplay Ability will activate both on Server & Client -- that is controlled via the Execution Policy.
+
+Keep in mind that Gameplay Abilities exist only on the locally controlled actors and on the server.  As such, you cannot replicate data meant to be visible on Simulated Proxies using Gameplay Abilities.  You would have to use other mechanisms, such as Attributes for replicated variables, or use RPC's directly on the Actors.
+
+### Replicated Variables in Gameplay Abilities
+
+The usage of replicated variables is deprecated as of UE5.5.  The deprecation warning is controlled by a Console Variable "AbilitySystem.DeprecateReplicatedProperties", so that users can turn off the warning and continue using the feature until they are ready to fix the issue.
+
+The reasoning is to prevent users from stumbling upon an impossible-to-solve bug regarding replication ordering:
+
+- Replicated variables are guaranteed to be delivered, but not in any particular order with respect to each other or RPC functions.
+- Gameplay Ability activation (and most synchronizing functions such as Target Data) rely on RPC's exchanged between the Client and Server.
+- Therefore, when executing an RPC (e.g. Gameplay Ability Activation) and performing operations on a replicated variable, you would never be guaranteed to have an up-to-date or stale value.
+
+For more information, see the [EDC article on object replication order](https://dev.epicgames.com/documentation/en-us/unreal-engine/replicated-object-execution-order-in-unreal-engine).
+
+If you believe you need a replicated variable, the solution is to instead use a Reliable RPC to send that data over.  Using a Reliable RPC will ensure proper ordering with the underlying synchronization mechanisms of GAS.
+
+### Remote Procedure Calls in Gameplay Abilities
+
+Remote Procedure Calls (RPC's) are the preferred method of communicating data between the client/server.  By making a Reliable RPC, you can ensure proper ordering with the other Gameplay Ability functions that support replication such as Activation.  There is currently no restriction against Unreliable RPC's, but know that order or delivery is not guaranteed.
+
+Using a Multicast RPC will produce a validation warning (typically visible when compiling the Blueprint).  Since Gameplay Abilities never exist on Simulated Proxies, Multicast RPC's make little sense in the context of a Gameplay Ability.
+
+### RPC Batching to Ensure Proper Activation and ReplicatedTargetData Order
+
+There is a trick Epic uses internally to bundle the Gameplay Ability activation and Replicated Target Data.  Normally, if one were to implement a Locally Predicted Gameplay Ability that calls Activate() which in turn sets Replicated Target Data, the two would arrive at the Server in separate RPC's:
+
+1. ServerTryActivateAbility (which will in turn call Activate)
+2. ServerSetReplicatedTargetData (which will then set the data to the desired value -- but Activate has already run!)
+
+There is a structure called [FScopedServerAbilityRPCBatcher](./Source/GameplayAbilities/Public/GameplayAbilityTypes.h) which is designed to use a single RPC to send both Activation and Target Data.  To use it, do the following:
+
+1. In your `UAbilitySystemComponent`-derived class, override `ShouldDoServerAbilityRPCBatch` to return true.
+2. In native code, Create an FScopedServerAbilityRPCBatcher on the stack.
+3. Activate your Ability through your desired function (e.g. TryActivateAbility).
+4. During the initial Activation of your ability, perform any CallServerSetReplicatedTargetData call.
+5. When the destructor of the FScopedServerAbilityRPCBatcher executes (by going out of scope), it will call a batched RPC that contains both the Activation and the ReplicatedTargetData.
+
+By using this structure, you will be guaranteed that the Server has the desired RPC data prior to calling the Gameplay Ability's activation function.
+
 ---
 
 # Gameplay Effects {#gameplay-effects}
diff --git a/Engine/Plugins/Runtime/GameplayAbilities/Source/GameplayAbilities/Private/Abilities/GameplayAbility.cpp b/Engine/Plugins/Runtime/GameplayAbilities/Source/GameplayAbilities/Private/Abilities/GameplayAbility.cpp
--- a/Engine/Plugins/Runtime/GameplayAbilities/Source/GameplayAbilities/Private/Abilities/GameplayAbility.cpp
+++ b/Engine/Plugins/Runtime/GameplayAbilities/Source/GameplayAbilities/Private/Abilities/GameplayAbility.cpp
@@ -33,6 +33,11 @@ namespace UE::AbilitySystem::Private
 
 	int32 AbilitySystemShowMakeOutgoingGameplayEffectSpecs = 0;
 	static FAutoConsoleVariableRef CVarAbilitySystemShowMakeOutgoingGameplayEffectSpecs(TEXT("AbilitySystem.ShowClientMakeOutgoingSpecs"), AbilitySystemShowMakeOutgoingGameplayEffectSpecs, TEXT("Displays all GameplayEffect specs created on non authority clients"), ECVF_Default );
+
+	/** We want to deprecate the usage of Replicated Properties due to replication order issues (see the GameplayAbilities/README.md for more information).  Default to 2 to not invalidate existing data during unattended runs (e.g. automated tests). */
+	int32 CVarDeprecateReplicatedPropertiesValue = 2;
+	static FAutoConsoleVariableRef CVarDeprecateReplicatedProperties(TEXT("AbilitySystem.DeprecateReplicatedProperties"), CVarDeprecateReplicatedPropertiesValue,
+		TEXT("Deprecate the use of Replicated Properties. ") TEXT("0: No. 1: Yes. 2: Yes, Except During Automation"), ECVF_Default );
 }
 
 namespace FAbilitySystemTweaks
@@ -231,6 +236,7 @@ bool UGameplayAbility::IsSupportedForNetworking() const
 EDataValidationResult UGameplayAbility::IsDataValid(FDataValidationContext& Context) const
 {
 	EDataValidationResult Result = EDataValidationResult::Valid;
+	const bool bIsLikelyRunningAutomation = IsRunningCommandlet() || FApp::IsUnattended();
 
 	if (!UE::AbilitySystem::Private::CVarAllowNonInstancedGAsValue)
 	{
@@ -241,21 +247,40 @@ EDataValidationResult UGameplayAbility::IsDataValid(FDataValidationContext& Cont
 		}
 	}
 
-	if (GetReplicationPolicy() == EGameplayAbilityReplicationPolicy::ReplicateNo)
+	const bool bIsReplicated = GetReplicationPolicy() == EGameplayAbilityReplicationPolicy::ReplicateYes;
+	if (bIsReplicated)
 	{
-		UBlueprintGeneratedClass* BPClass = Cast<UBlueprintGeneratedClass>(GetClass());
-		if (BPClass && BPClass->NumReplicatedProperties > 0)
+		if (GetInstancingPolicy() != EGameplayAbilityInstancingPolicy::InstancedPerActor)
 		{
-			Context.AddError(LOCTEXT("ReplicatedVariablesNeedReplicationYes", "Gameplay Ability Blueprint has replicated variables but Replication Policy is set to not replicate"));
+			Context.AddError(FText::Format(LOCTEXT("ReplicatedInstancePolicyNotSupported", "Instancing Policy '{0}' is not supported for Replication.  Either change the Replication Policy or the Instancing Policy"), UEnum::GetDisplayValueAsText(GetInstancingPolicy())));
 			Result = EDataValidationResult::Invalid;
 		}
 	}
+
+	// Find out which properties are replicated
+	if (UE::AbilitySystem::Private::CVarDeprecateReplicatedPropertiesValue > 0)
+	{
+		const UClass* Class = GetClass();
+		for (const FProperty* Property = Class->PropertyLink; Property; Property = Property->PropertyLinkNext)
+		{
+			if (Property->RepIndex > 0 || Property->HasAnyPropertyFlags(CPF_RepNotify | CPF_Net))
+			{
+				Context.AddWarning(FText::Format(LOCTEXT("ReplicatedPropertiesAreDeprecated", "{0}: Replicated properties are deprecated in Gameplay Abilities (see GameplayAbilities/README.md). Use Reliable RPCs for variable replication."), Property->GetDisplayNameText()));
+				const bool bAllowSuppressErrorsSoAutomationCanPass = (UE::AbilitySystem::Private::CVarDeprecateReplicatedPropertiesValue > 1);
+				Result = (bAllowSuppressErrorsSoAutomationCanPass && bIsLikelyRunningAutomation) ? CombineDataValidationResults(Result, EDataValidationResult::Valid) : EDataValidationResult::Invalid;
+			}
+		}
+	}
 	else
 	{
-		if (GetInstancingPolicy() != EGameplayAbilityInstancingPolicy::InstancedPerActor)
+		if (GetReplicationPolicy() == EGameplayAbilityReplicationPolicy::ReplicateNo)
 		{
-			Context.AddError(FText::Format(LOCTEXT("ReplicatedInstancePolicyNotSupported", "Instancing Policy '{0}' is not supported for Replication.  Either change the Replication Policy or the Instancing Policy"), UEnum::GetDisplayValueAsText(GetInstancingPolicy())));
-			Result = EDataValidationResult::Invalid;
+			UBlueprintGeneratedClass* BPClass = Cast<UBlueprintGeneratedClass>(GetClass());
+			if (BPClass && BPClass->NumReplicatedProperties > 0)
+			{
+				Context.AddError(LOCTEXT("ReplicatedVariablesNeedReplicationYes", "Gameplay Ability Blueprint has replicated variables but Replication Policy is set to not replicate"));
+				Result = EDataValidationResult::Invalid;
+			}
 		}
 	}
 
@@ -267,6 +292,20 @@ EDataValidationResult UGameplayAbility::IsDataValid(FDataValidationContext& Cont
 			Context.AddError(ErrorText);
 			Result = EDataValidationResult::Invalid;
 		}
+		else if (FuncIter->HasAnyFunctionFlags(EFunctionFlags::FUNC_Net) && !bIsReplicated)
+		{
+			FText ErrorText = FText::Format(LOCTEXT("RpcRequiresReplicationYes", "{0}: RPC Functions require ReplicationPolicy to be ReplicateYes in order to actually work."), FuncIter->GetDisplayNameText());
+			if (bIsLikelyRunningAutomation)
+			{
+				Context.AddWarning(ErrorText);
+				Result = CombineDataValidationResults(Result, EDataValidationResult::Valid);
+			}
+			else
+			{
+				Context.AddError(ErrorText);
+				Result = EDataValidationResult::Invalid;
+			}
+		}
 	}
 
 	return Result;
