[GAS] Relabel the Gameplay Cue Notify Actor functions to make it clear what they're used for and how they relate to replication.  OnActive becomes OnBurst.  WhileActive becomes OnBecomeRelevant.  OnRemove becomes OnCeaseRelevant.  Introduce some data validation to catch tricky bugs with Gameplay Cues "stuck on" due to OnActive/OnRemove which is now more clearly marked incorrect.

[RN] Gameplay Cue Notify Actor (Looping) Function names have been renamed to give users a better understanding of when they trigger. An asset validation warning has been added in the case where a user implements OnActive and OnRemove, but not WhileActive. These have now been named to OnBurst, OnCeaseRelevant, and OnBecomeRelevant respectively.

diff --git a/Engine/Plugins/Runtime/GameplayAbilities/README.md b/Engine/Plugins/Runtime/GameplayAbilities/README.md
--- a/Engine/Plugins/Runtime/GameplayAbilities/README.md
+++ b/Engine/Plugins/Runtime/GameplayAbilities/README.md
@@ -184,7 +184,7 @@ The Gameplay Ability System uses Gameplay Tags extensively throughout.  See the
 
 ---
 
-# Gameplay Cues (#gameplay-cues)
+# Gameplay Cues {#gameplay-cues}
 
 Gameplay Cues are a system for decoupling visual and audio fx from gameplay code.  On start-up, special Gameplay Cue asset folders are scanned for [Gameplay Cue Sets](./Source/GameplayAbilities/Public/GameplayCueSet.h), and *Gameplay Cue Notify* classes.
 
@@ -198,6 +198,36 @@ Due to these Gameplay Cues needing to obey network relevancy (i.e. far away play
 
 An advanced form of replication proxies also exists for the property replication so it may follow the same relevancy rules.  See `FMinimalGameplayCueReplicationProxy` in the [GameplayCueInterface](./Source/GameplayAbilities/Public/GameplayCueInterface.h).
 
+Due to the Burst Cues being replicated by RPC and the Looping Cues being replicated by replicated variables, one can run into an issue where the unreliable burst RPC gets dropped but the looping events (OnBecomeRelevant/OnCeaseRelevant) arrive.  Less obvious, the unreliable OnBurst RPC can arrive but the OnBecomeRelevant/OnCeaseRelevant can be dropped if the Cue is removed on the server quick enough to result in no state changes for network serialization.
+
+See the section on [Gameplay Cue Events](#gc-events) below for guidelines on how to implement your Gameplay Cue while taking into consideration network replication.
+
+## Gameplay Cue Events {#gc-events}
+
+When implementing a Gameplay Cue Notify Actor, the (legacy) naming of the functions may be confusing.  In UE5.5 the Blueprint (user-facing) names have changed in order to better represent what each function does.  They are laid out below.
+
+### OnExecute
+
+The execute function is the easiest to reason about:  It happens when you *Execute* a one-shot Gameplay Cue (aka a Static Notify / non-Looping Gameplay Cue).  The code path to Execute a Gameplay Cue (for Static Notifies) is different than the code path to Add a Gameplay Cue (for Looping Gameplay Cues aka Actor Notifies).
+
+Due to the code path for execution being different, the caller of the Gameplay Cue must know that the receiver of the Gameplay Cue is a Static Notify in order for this to execute properly.  The call should route through ExecuteGameplayCue see [GameplayCueFunctionLibrary](./Source/GameplayAbilities/Public/GameplayCueFunctionLibrary.h).
+
+### OnBurst (native: OnActive)
+
+This event executes only once when a *Looping Gameplay Cue* first fires.  Due to it being delivered by unreliable RPC, it can be dropped silently by a client.  You can use this to implement cosmetic effects that are only relevant if a client witnessed the Gameplay Cue triggering.
+
+### OnBecomeRelevant (native: WhileActive)
+
+This event executes when the *Looping Gameplay Cue* first comes into network relevancy (usually when it's first added).  For instance, PawnA can have a Gameplay Cue activated, PawnB can join the game and still receive PawnA's OnBecomeRelevant -- but not receive OnBurst.
+
+This is important to understand as OnBecomeRelevant and OnCeaseRelevant are both guaranteed to fire on the same Cue, whereas OnBurst is not guaranteed.
+
+### OnCeaseRelevant (native: OnRemove)
+
+This event executes when the *Looping Gameplay Cue* gets removed from network relevancy.  Usually that's when the server executes the removal of the Cue, but could also be when the client loses relevancy (e.g. by distance) of the viewed Cue.
+
+In UE5.5, a warning is introduced if a Gameplay Cue implements OnBurst and OnCeaseRelevant and not OnBecomeRelevant.  The reasoning is that the opposite of OnCeaseRelevant is OnBecomeRelevant, not OnBurst and it's likely that the old naming scheme (OnActive/OnRemove) was a source of confusion.
+
 ---
 
 # How Gameplay Prediction Works
diff --git a/Engine/Plugins/Runtime/GameplayAbilities/Source/GameplayAbilities/Private/GameplayCueNotify_Actor.cpp b/Engine/Plugins/Runtime/GameplayAbilities/Source/GameplayAbilities/Private/GameplayCueNotify_Actor.cpp
--- a/Engine/Plugins/Runtime/GameplayAbilities/Source/GameplayAbilities/Private/GameplayCueNotify_Actor.cpp
+++ b/Engine/Plugins/Runtime/GameplayAbilities/Source/GameplayAbilities/Private/GameplayCueNotify_Actor.cpp
@@ -10,8 +10,14 @@
 #include "GameplayCueManager.h"
 #include "GameplayCueNotifyTypes.h"
 
+#if WITH_EDITOR
+#include "Engine/BlueprintGeneratedClass.h"
+#include "Misc/DataValidation.h"
+#endif
+
 #include UE_INLINE_GENERATED_CPP_BY_NAME(GameplayCueNotify_Actor)
 
+#define LOCTEXT_NAMESPACE "GameplayCueNotify_Actor"
 
 namespace FAbilitySystemTweaks
 {
@@ -91,6 +97,29 @@ void AGameplayCueNotify_Actor::Destroyed()
 }
 
 #if WITH_EDITOR
+EDataValidationResult AGameplayCueNotify_Actor::IsDataValid(FDataValidationContext& Context) const
+{
+	EDataValidationResult ReturnValue = Super::IsDataValid(Context);
+
+	auto IsBlueprintOverride = [&](const FName& FuncName) -> bool
+	{
+		const UFunction* Func = GetClass()->FindFunctionByName(FuncName);
+		return Func && ensure(Func->GetOuter())
+			&& Func->GetOuter()->IsA(UBlueprintGeneratedClass::StaticClass());
+	};
+
+	bool bHasOnActiveInBP = IsBlueprintOverride(GET_FUNCTION_NAME_CHECKED(AGameplayCueNotify_Actor, OnActive));
+	bool bHasWhileActiveInBP = IsBlueprintOverride(GET_FUNCTION_NAME_CHECKED(AGameplayCueNotify_Actor, WhileActive));
+	bool bHasOnRemoveInBP = IsBlueprintOverride(GET_FUNCTION_NAME_CHECKED(AGameplayCueNotify_Actor, OnRemove));
+
+	if (bHasOnActiveInBP && bHasOnRemoveInBP && !bHasWhileActiveInBP)
+	{
+		Context.AddWarning(LOCTEXT("OnActiveAndOnRemove", "OnCeaseRelevant implemented without OnBecomeRelevant. OnBurst is likely a mistake (if not, implement a blank OnBecomeRelevant)."));
+	}
+
+	return ReturnValue;
+}
+
 void AGameplayCueNotify_Actor::PostEditChangeProperty(FPropertyChangedEvent& PropertyChangedEvent)
 {
 	const FProperty* PropertyThatChanged = PropertyChangedEvent.Property;
@@ -411,3 +440,4 @@ void AGameplayCueNotify_Actor::ReuseAfterRecycle()
 	SetActorHiddenInGame(false);
 }
 
+#undef LOCTEXT_NAMESPACE
diff --git a/Engine/Plugins/Runtime/GameplayAbilities/Source/GameplayAbilities/Public/GameplayCueNotify_Actor.h b/Engine/Plugins/Runtime/GameplayAbilities/Source/GameplayAbilities/Public/GameplayCueNotify_Actor.h
--- a/Engine/Plugins/Runtime/GameplayAbilities/Source/GameplayAbilities/Public/GameplayCueNotify_Actor.h
+++ b/Engine/Plugins/Runtime/GameplayAbilities/Source/GameplayAbilities/Public/GameplayCueNotify_Actor.h
@@ -33,6 +33,10 @@ protected:
 	virtual void K2_DestroyActor() override;
 	virtual void Destroyed() override;
 
+#if WITH_EDITOR
+	virtual EDataValidationResult IsDataValid(class FDataValidationContext& Context) const override;
+#endif
+
 public:
 	virtual void SetOwner( AActor* NewOwner ) override;
 
@@ -86,15 +90,15 @@ public:
 	bool OnExecute(AActor* MyTarget, const FGameplayCueParameters& Parameters);
 
 	/** Called when a GameplayCue with duration is first activated, this will only be called if the client witnessed the activation */
-	UFUNCTION(BlueprintNativeEvent, Category = "GameplayCueNotify")
+	UFUNCTION(BlueprintNativeEvent, Category = "GameplayCueNotify", DisplayName="On Burst", meta=(ScriptName="OnBurst;OnActive"))
 	bool OnActive(AActor* MyTarget, const FGameplayCueParameters& Parameters);
 
 	/** Called when a GameplayCue with duration is first seen as active, even if it wasn't actually just applied (Join in progress, etc) */
-	UFUNCTION(BlueprintNativeEvent, Category = "GameplayCueNotify")
+	UFUNCTION(BlueprintNativeEvent, Category = "GameplayCueNotify", DisplayName="On Become Relevant", meta=(ScriptName="OnBecomeRelevant;WhileActive"))
 	bool WhileActive(AActor* MyTarget, const FGameplayCueParameters& Parameters);
 
 	/** Called when a GameplayCue with duration is removed */
-	UFUNCTION(BlueprintNativeEvent, Category = "GameplayCueNotify")
+	UFUNCTION(BlueprintNativeEvent, Category = "GameplayCueNotify", DisplayName="On Cease Relevant", meta=(ScriptName="OnCeaseRelevant;OnRemove"))
 	bool OnRemove(AActor* MyTarget, const FGameplayCueParameters& Parameters);
 
 	/** Tag this notify is activated by */
@@ -128,17 +132,17 @@ public:
 	bool bUniqueInstancePerSourceObject;
 
 	/**
-	 *	Does this cue trigger its OnActive event if it's already been triggered?
+	 *	Does this cue trigger its On Burst event if it's already been triggered?
 	 *  This can occur when the associated tag is triggered by multiple sources and there is no unique instancing.
 	 */ 	 
-	UPROPERTY(EditDefaultsOnly, Category = GameplayCue)
+	UPROPERTY(EditDefaultsOnly, Category = GameplayCue, DisplayName="Allow Multiple On Burst Events")
 	bool bAllowMultipleOnActiveEvents;
 
 	/**
-	 *	Does this cue trigger its WhileActive event if it's already been triggered?
+	 *	Does this cue trigger its On Become Relevant event if it's already been triggered?
 	 *  This can occur when the associated tag is triggered by multiple sources and there is no unique instancing.
-	 */ 	 
-	UPROPERTY(EditDefaultsOnly, Category = GameplayCue)
+	 */
+	UPROPERTY(EditDefaultsOnly, Category = GameplayCue, DisplayName="Allow Multiple On Become Relevant Events")
 	bool bAllowMultipleWhileActiveEvents;
 
 	/** How many instances of the gameplay cue to preallocate */
